Copied CL 11188915 from //UE4/Dev-VirtualProduction.

GeometryCache from Alembic file without importing. (Fixed for Mac build)

diff --git a/Engine/Plugins/Experimental/GeometryCache/Source/GeometryCache/Private/GeometryCacheHelpers.h b/Engine/Plugins/Experimental/GeometryCache/Source/GeometryCache/Private/GeometryCacheHelpers.h
deleted file mode 100644
diff --git a/Engine/Plugins/Experimental/GeometryCache/Source/GeometryCache/Public/GeometryCacheHelpers.h b/Engine/Plugins/Experimental/GeometryCache/Source/GeometryCache/Public/GeometryCacheHelpers.h
new file mode 100644
--- /dev/null
+++ b/Engine/Plugins/Experimental/GeometryCache/Source/GeometryCache/Public/GeometryCacheHelpers.h
@@ -0,0 +1,14 @@
+// Copyright Epic Games, Inc. All Rights Reserved.
+
+#pragma once
+
+struct GeometyCacheHelpers 
+{
+	/**
+		Use this instead of fmod when working with looping animations as fmod gives incorrect results when using negative times.
+	*/
+	static inline float WrapAnimationTime(float Time, float Duration)
+	{
+		return Time - Duration * FMath::FloorToFloat(Time / Duration);
+	}
+};
\ No newline at end of file
diff --git a/Engine/Plugins/Experimental/GeometryCache/Source/GeometryCacheTracks/Private/MovieSceneGeometryCacheTemplate.cpp b/Engine/Plugins/Experimental/GeometryCache/Source/GeometryCacheTracks/Private/MovieSceneGeometryCacheTemplate.cpp
--- a/Engine/Plugins/Experimental/GeometryCache/Source/GeometryCacheTracks/Private/MovieSceneGeometryCacheTemplate.cpp
+++ b/Engine/Plugins/Experimental/GeometryCache/Source/GeometryCacheTracks/Private/MovieSceneGeometryCacheTemplate.cpp
@@ -87,13 +87,20 @@ struct FGeometryCacheExecutionToken
 					UGeometryCacheComponent* GeometryComp = GeometryMeshComponentFromObject(Obj);
 					if (GeometryComp)
 					{
-						if (Params.GeometryCacheAsset != GeometryComp->GetGeometryCache())
+						// Set the GeometryCache on the component only if it's set and valid in the Params
+						if (Params.GeometryCacheAsset && Params.GeometryCacheAsset != GeometryComp->GetGeometryCache())
 						{
 							UGeometryCache* GeomCache = Params.GeometryCacheAsset;
 							{
 								GeometryComp->SetGeometryCache(GeomCache);
 							}
 						}
+						else
+						{
+							// It could be unset if the Params was referencing a transient GeometryCache
+							// In that case, use the GeometryCache that is set on the component
+							Params.GeometryCacheAsset = GeometryComp->GetGeometryCache();
+						}
 						Player.SavePreAnimatedState(*GeometryComp, FPreAnimatedGeometryCacheTokenProducer::GetAnimTypeID(), FPreAnimatedGeometryCacheTokenProducer());
 						GeometryComp->SetManualTick(true);
 						// calculate the time at which to evaluate the animation
diff --git a/Engine/Plugins/Experimental/GeometryCacheAbcFile/GeometryCacheAbcFile.uplugin b/Engine/Plugins/Experimental/GeometryCacheAbcFile/GeometryCacheAbcFile.uplugin
new file mode 100644
--- /dev/null
+++ b/Engine/Plugins/Experimental/GeometryCacheAbcFile/GeometryCacheAbcFile.uplugin
@@ -0,0 +1,39 @@
+{
+	"FileVersion" : 1,
+	"Version" : 1,
+	"VersionName" : "0.1",
+	"FriendlyName" : "Geometry Cache from Alembic File(Experimental)",
+	"Description" : "Support Geometry Cache from Alembic file without importing",
+	"Category" : "Rendering",
+	"CreatedBy" : "Epic Games, Inc.",
+	"CreatedByURL" : "http://epicgames.com",
+	"DocsURL" : "",
+	"MarketplaceURL" : "",
+	"SupportURL" : "",
+	"EnabledByDefault" : false,
+	"CanContainContent" : false,
+	"IsBetaVersion" : true,
+	"Installed" : false,
+	"Modules": [
+		{
+			"Name": "GeometryCacheAbcFile",
+			"Type": "Editor",
+			"LoadingPhase": "Default",
+			"WhitelistPlatforms": [
+				"Win64",
+				"Mac",
+				"Linux"
+			]
+		}
+	],
+	"Plugins": [
+		{
+			"Name": "AlembicImporter",
+			"Enabled": true
+		},
+		{
+			"Name": "GeometryCache",
+			"Enabled": true
+		}
+	]
+}
diff --git a/Engine/Plugins/Experimental/GeometryCacheAbcFile/Source/GeometryCacheAbcFile/GeometryCacheAbcFile.Build.cs b/Engine/Plugins/Experimental/GeometryCacheAbcFile/Source/GeometryCacheAbcFile/GeometryCacheAbcFile.Build.cs
new file mode 100644
--- /dev/null
+++ b/Engine/Plugins/Experimental/GeometryCacheAbcFile/Source/GeometryCacheAbcFile/GeometryCacheAbcFile.Build.cs
@@ -0,0 +1,21 @@
+// Copyright Epic Games, Inc. All Rights Reserved.
+
+using UnrealBuildTool;
+
+public class GeometryCacheAbcFile : ModuleRules
+{
+	public GeometryCacheAbcFile(ReadOnlyTargetRules Target) : base(Target)
+	{
+		PublicDependencyModuleNames.AddRange(
+			new string[] {
+				"AlembicLibrary",
+				"Core",
+				"CoreUObject",
+				"Engine",
+				"GeometryCache",
+				"RenderCore",
+				"RHI"
+			}
+		);
+	}
+}
diff --git a/Engine/Plugins/Experimental/GeometryCacheAbcFile/Source/GeometryCacheAbcFile/Private/GeometryCacheAbcFileComponent.cpp b/Engine/Plugins/Experimental/GeometryCacheAbcFile/Source/GeometryCacheAbcFile/Private/GeometryCacheAbcFileComponent.cpp
new file mode 100644
--- /dev/null
+++ b/Engine/Plugins/Experimental/GeometryCacheAbcFile/Source/GeometryCacheAbcFile/Private/GeometryCacheAbcFileComponent.cpp
@@ -0,0 +1,67 @@
+// Copyright Epic Games, Inc. All Rights Reserved.
+
+#include "GeometryCacheAbcFileComponent.h"
+#include "GeometryCache.h"
+#include "GeometryCacheTrackAbcFile.h"
+#include "GeometryCacheAbcFileSceneProxy.h"
+
+#define LOCTEXT_NAMESPACE "GeometryCacheAbcFileComponent"
+
+#if WITH_EDITOR
+void UGeometryCacheAbcFileComponent::PostEditChangeProperty(FPropertyChangedEvent& PropertyChangedEvent)
+{
+	FProperty* PropertyThatChanged = PropertyChangedEvent.MemberProperty;
+	const FName PropertyName = PropertyThatChanged ? PropertyThatChanged->GetFName() : NAME_None;
+
+	if (PropertyName == GET_MEMBER_NAME_CHECKED(UGeometryCacheAbcFileComponent, AlembicFilePath))
+	{
+		InitializeGeometryCache();
+		InvalidateTrackSampleIndices();
+	}
+
+	Super::PostEditChangeProperty(PropertyChangedEvent);
+}
+#endif
+
+void UGeometryCacheAbcFileComponent::InitializeGeometryCache()
+{
+	UGeometryCacheTrackAbcFile* AbcFileTrack = nullptr;
+	if (!GeometryCache)
+	{
+		// Transient GeometryCache for use in current session
+		GeometryCache = NewObject<UGeometryCache>();
+		AbcFileTrack = NewObject<UGeometryCacheTrackAbcFile>(GeometryCache);
+		GeometryCache->AddTrack(AbcFileTrack);
+	}
+	else
+	{
+		AbcFileTrack = Cast<UGeometryCacheTrackAbcFile>(GeometryCache->Tracks[0]);
+	}
+
+	// #ueent_todo: Should be able to clear the Alembic file
+	if (!AlembicFilePath.FilePath.IsEmpty() && (AlembicFilePath.FilePath != AbcFileTrack->GetSourceFile()))
+	{
+		bool bIsValid = AbcFileTrack->SetSourceFile(AlembicFilePath.FilePath);
+		if (bIsValid)
+		{
+			// Also store the number of frames in the cache
+			GeometryCache->SetFrameStartEnd(0, AbcFileTrack->GetEndFrameIndex());
+		}
+	}
+
+	MarkRenderStateDirty();
+}
+
+void UGeometryCacheAbcFileComponent::PostLoad()
+{
+	Super::PostLoad();
+
+	InitializeGeometryCache();
+}
+
+FPrimitiveSceneProxy* UGeometryCacheAbcFileComponent::CreateSceneProxy()
+{
+	return new FGeometryCacheAbcFileSceneProxy(this);
+}
+
+#undef LOCTEXT_NAMESPACE
diff --git a/Engine/Plugins/Experimental/GeometryCacheAbcFile/Source/GeometryCacheAbcFile/Private/GeometryCacheAbcFileModule.cpp b/Engine/Plugins/Experimental/GeometryCacheAbcFile/Source/GeometryCacheAbcFile/Private/GeometryCacheAbcFileModule.cpp
new file mode 100644
--- /dev/null
+++ b/Engine/Plugins/Experimental/GeometryCacheAbcFile/Source/GeometryCacheAbcFile/Private/GeometryCacheAbcFileModule.cpp
@@ -0,0 +1,10 @@
+// Copyright Epic Games, Inc. All Rights Reserved.
+
+#include "Modules/ModuleInterface.h"
+#include "Modules/ModuleManager.h"
+
+class FGeometryCacheAbcFileModule : public IModuleInterface
+{
+};
+
+IMPLEMENT_MODULE(FGeometryCacheAbcFileModule, GeometryCacheAbcFile);
\ No newline at end of file
diff --git a/Engine/Plugins/Experimental/GeometryCacheAbcFile/Source/GeometryCacheAbcFile/Private/GeometryCacheAbcFileSceneProxy.cpp b/Engine/Plugins/Experimental/GeometryCacheAbcFile/Source/GeometryCacheAbcFile/Private/GeometryCacheAbcFileSceneProxy.cpp
new file mode 100644
--- /dev/null
+++ b/Engine/Plugins/Experimental/GeometryCacheAbcFile/Source/GeometryCacheAbcFile/Private/GeometryCacheAbcFileSceneProxy.cpp
@@ -0,0 +1,65 @@
+// Copyright Epic Games, Inc. All Rights Reserved.
+
+#include "GeometryCacheAbcFileSceneProxy.h"
+#include "GeometryCacheAbcFileComponent.h"
+#include "GeometryCacheTrackAbcFile.h"
+
+FGeometryCacheAbcFileSceneProxy::FGeometryCacheAbcFileSceneProxy(UGeometryCacheAbcFileComponent* Component)
+: FGeometryCacheSceneProxy(Component, [this]() { return new FGeomCacheTrackAbcFileProxy(GetScene().GetFeatureLevel()); })
+{
+}
+
+bool FGeomCacheTrackAbcFileProxy::UpdateMeshData(float Time, bool bLooping, int32& InOutMeshSampleIndex, FGeometryCacheMeshData& OutMeshData)
+{
+	if (UGeometryCacheTrackAbcFile* AbcTrack = Cast<UGeometryCacheTrackAbcFile>(Track))
+	{
+		FGeometryCacheMeshData* PtrMeshData = nullptr;
+		bool bResult = AbcTrack->UpdateMeshData(Time, bLooping, InOutMeshSampleIndex, PtrMeshData);
+		if (bResult)
+		{
+			OutMeshData = *PtrMeshData;
+		}
+		return bResult;
+	}
+	return false;
+}
+
+bool FGeomCacheTrackAbcFileProxy::GetMeshData(int32 SampleIndex, FGeometryCacheMeshData& OutMeshData)
+{
+	if (UGeometryCacheTrackAbcFile* AbcTrack = Cast<UGeometryCacheTrackAbcFile>(Track))
+	{
+		return AbcTrack->GetMeshData(SampleIndex, OutMeshData);
+	}
+	return false;
+}
+
+bool FGeomCacheTrackAbcFileProxy::IsTopologyCompatible(int32 SampleIndexA, int32 SampleIndexB)
+{
+	// No support for interpolation for now (assume the topology is variable)
+	return false;
+}
+
+const FVisibilitySample& FGeomCacheTrackAbcFileProxy::GetVisibilitySample(float Time, const bool bLooping) const
+{
+	// Assume the track is visible for its whole duration
+	return FVisibilitySample::VisibleSample;
+}
+
+void FGeomCacheTrackAbcFileProxy::FindSampleIndexesFromTime(float Time, bool bLooping, bool bIsPlayingBackwards, int32 &OutFrameIndex, int32 &OutNextFrameIndex, float &InterpolationFactor)
+{
+	if (UGeometryCacheTrackAbcFile* AbcTrack = Cast<UGeometryCacheTrackAbcFile>(Track))
+	{
+		int32 ThisFrameIndex = AbcTrack->FindSampleIndexFromTime(Time, bLooping);
+		int32 LastFrameIndex = AbcTrack->GetEndFrameIndex();
+		OutFrameIndex = ThisFrameIndex;
+		OutNextFrameIndex = OutFrameIndex + 1;
+		InterpolationFactor = 0.f;
+
+		// If playing backwards the logical order of previous and next is reversed
+		if (bIsPlayingBackwards)
+		{
+			Swap(OutFrameIndex, OutNextFrameIndex);
+			InterpolationFactor = 1.0f - InterpolationFactor;
+		}
+	}
+}
diff --git a/Engine/Plugins/Experimental/GeometryCacheAbcFile/Source/GeometryCacheAbcFile/Private/GeometryCacheTrackAbcFile.cpp b/Engine/Plugins/Experimental/GeometryCacheAbcFile/Source/GeometryCacheAbcFile/Private/GeometryCacheTrackAbcFile.cpp
new file mode 100644
--- /dev/null
+++ b/Engine/Plugins/Experimental/GeometryCacheAbcFile/Source/GeometryCacheAbcFile/Private/GeometryCacheTrackAbcFile.cpp
@@ -0,0 +1,146 @@
+// Copyright Epic Games, Inc. All Rights Reserved.
+
+#include "GeometryCacheTrackAbcFile.h"
+
+#include "AbcImporter.h"
+#include "AbcUtilities.h"
+#include "AbcImportSettings.h"
+#include "GeometryCacheHelpers.h"
+
+UGeometryCacheTrackAbcFile::UGeometryCacheTrackAbcFile()
+: EndFrameIndex(0)
+{
+}
+
+const bool UGeometryCacheTrackAbcFile::UpdateMeshData(const float Time, const bool bLooping, int32& InOutMeshSampleIndex, FGeometryCacheMeshData*& OutMeshData)
+{
+	const int32 SampleIndex = FindSampleIndexFromTime(Time, bLooping);
+
+	// If InOutMeshSampleIndex equals -1 (first creation) update the OutVertices and InOutMeshSampleIndex
+	// Update the Vertices and Index if SampleIndex is different from the stored InOutMeshSampleIndex
+	if (InOutMeshSampleIndex == -1 || SampleIndex != InOutMeshSampleIndex)
+	{
+		if (GetMeshData(SampleIndex, MeshData))
+		{
+			OutMeshData = &MeshData;
+			InOutMeshSampleIndex = SampleIndex;
+			return true;
+		}
+	}
+	return false;
+}
+
+const bool UGeometryCacheTrackAbcFile::UpdateBoundsData(const float Time, const bool bLooping, const bool bIsPlayingBackward, int32& InOutBoundsSampleIndex, FBox& OutBounds)
+{
+	const int32 SampleIndex = FindSampleIndexFromTime(Time, bLooping);
+
+	FGeometryCacheTrackSampleInfo SampledInfo = GetSampleInfo(Time, bLooping);
+	if (InOutBoundsSampleIndex != SampleIndex)
+	{
+		OutBounds = SampledInfo.BoundingBox;
+		InOutBoundsSampleIndex = SampleIndex;
+		return true;
+	}
+	return false;
+}
+
+bool UGeometryCacheTrackAbcFile::SetSourceFile(const FString& FilePath)
+{
+	if (!FilePath.IsEmpty())
+	{
+		AbcFile = MakeUnique<FAbcFile>(FilePath);
+		EAbcImportError Result = AbcFile->Open();
+
+		if (Result != EAbcImportError::AbcImportError_NoError)
+		{
+			AbcFile.Reset();
+			return false;
+		}
+
+		// #ueent_todo: Expose some conversion settings in the component to use with the AbcFile
+		UAbcImportSettings* ImportSettings = DuplicateObject(GetMutableDefault<UAbcImportSettings>(), GetTransientPackage());
+
+		EndFrameIndex = FMath::Max(AbcFile->GetMaxFrameIndex() - 1, 1);
+		ImportSettings->SamplingSettings.FrameEnd = EndFrameIndex;
+		ImportSettings->ImportType = EAlembicImportType::GeometryCache;
+
+		Result = AbcFile->Import(ImportSettings);
+
+		if (Result != EAbcImportError::AbcImportError_NoError)
+		{
+			AbcFile.Reset();
+			return false;
+		}
+
+		TArray<FMatrix> Mats;
+		Mats.Add(FMatrix::Identity);
+		Mats.Add(FMatrix::Identity);
+
+		TArray<float> MatTimes;
+		MatTimes.Add(0.0f);
+		MatTimes.Add(AbcFile->GetImportLength() + AbcFile->GetImportTimeOffset());
+		SetMatrixSamples(Mats, MatTimes);
+
+		Duration = AbcFile->GetImportLength();
+
+		// Fill out the MeshData with the sample from time 0
+		GetMeshData(0, MeshData);
+	}
+	else
+	{
+		AbcFile.Reset();
+
+		EndFrameIndex = 0;
+		Duration = 0.f;
+
+		MatrixSamples.Empty();
+		MatrixSampleTimes.Empty();
+	}
+	SourceFile = FilePath;
+	return true;
+}
+
+const int32 UGeometryCacheTrackAbcFile::FindSampleIndexFromTime(const float Time, const bool bLooping) const
+{
+	if (AbcFile)
+	{
+		float SampleTime = Time;
+		if (bLooping)
+		{
+			SampleTime = GeometyCacheHelpers::WrapAnimationTime(Time, Duration);
+		}
+		return AbcFile->GetFrameIndex(SampleTime);
+	}
+	return 0;
+}
+
+const FGeometryCacheTrackSampleInfo& UGeometryCacheTrackAbcFile::GetSampleInfo(float Time, bool bLooping)
+{
+	float SampleTime = Time;
+	if (bLooping)
+	{
+		SampleTime = GeometyCacheHelpers::WrapAnimationTime(Time, Duration);
+	}
+
+	// #ueent_todo: Return the correct SampleInfo without double querying the Alembic. 
+	// Currently returns the info for MeshData queried in UpdateMeshData
+	SampleInfo = FGeometryCacheTrackSampleInfo(
+		SampleTime,
+		MeshData.BoundingBox,
+		MeshData.Positions.Num(),
+		MeshData.Indices.Num()
+	);
+
+	return SampleInfo;
+}
+
+bool UGeometryCacheTrackAbcFile::GetMeshData(int32 SampleIndex, FGeometryCacheMeshData& OutMeshData)
+{
+	if (AbcFile.IsValid())
+	{
+		// #ueent_todo: Implement optimized Alembic querying
+		FAbcUtilities::GetFrameMeshData(*AbcFile, SampleIndex, OutMeshData);
+		return true;
+	}
+	return false;
+}
diff --git a/Engine/Plugins/Experimental/GeometryCacheAbcFile/Source/GeometryCacheAbcFile/Public/GeometryCacheAbcFileComponent.h b/Engine/Plugins/Experimental/GeometryCacheAbcFile/Source/GeometryCacheAbcFile/Public/GeometryCacheAbcFileComponent.h
new file mode 100644
--- /dev/null
+++ b/Engine/Plugins/Experimental/GeometryCacheAbcFile/Source/GeometryCacheAbcFile/Public/GeometryCacheAbcFileComponent.h
@@ -0,0 +1,38 @@
+// Copyright Epic Games, Inc. All Rights Reserved.
+
+#pragma once
+
+#include "CoreMinimal.h"
+#include "GeometryCacheComponent.h"
+
+#include "GeometryCacheAbcFileComponent.generated.h"
+
+class UGeometryCache;
+
+/** GeometryCacheAbcFileComponent, encapsulates a transient GeometryCache asset instance that fetches its data from an Alembic file and implements functionality for rendering and playback */
+UCLASS(ClassGroup = (Rendering), meta = (BlueprintSpawnableComponent, DisplayName = "Geometry Cache Alembic File"), Experimental, ClassGroup = Experimental)
+class GEOMETRYCACHEABCFILE_API UGeometryCacheAbcFileComponent : public UGeometryCacheComponent
+{
+	GENERATED_BODY()
+
+public:
+
+	UPROPERTY(EditAnywhere, Category = "Alembic", meta = (FilePathFilter = "Alembic files (*.abc)|*.abc"))
+	FFilePath AlembicFilePath;
+
+public:
+#if WITH_EDITOR
+	virtual void PostEditChangeProperty(FPropertyChangedEvent& PropertyChangedEvent) override;
+#endif // WITH_EDITOR
+
+	//~ Begin UObject Interface
+	virtual void PostLoad() override;
+	//~ End UObject Interface
+
+	//~ Begin UPrimitiveComponent Interface.
+	virtual FPrimitiveSceneProxy* CreateSceneProxy() override;
+	//~ End UPrimitiveComponent Interface.
+
+protected:
+	void InitializeGeometryCache();
+};
diff --git a/Engine/Plugins/Experimental/GeometryCacheAbcFile/Source/GeometryCacheAbcFile/Public/GeometryCacheAbcFileSceneProxy.h b/Engine/Plugins/Experimental/GeometryCacheAbcFile/Source/GeometryCacheAbcFile/Public/GeometryCacheAbcFileSceneProxy.h
new file mode 100644
--- /dev/null
+++ b/Engine/Plugins/Experimental/GeometryCacheAbcFile/Source/GeometryCacheAbcFile/Public/GeometryCacheAbcFileSceneProxy.h
@@ -0,0 +1,30 @@
+// Copyright Epic Games, Inc. All Rights Reserved.
+
+#pragma once
+
+#include "CoreMinimal.h"
+#include "GeometryCacheSceneProxy.h"
+
+/** Scene proxy that uses FGeomCacheTrackAbcFileProxy */
+class FGeometryCacheAbcFileSceneProxy final : public FGeometryCacheSceneProxy
+{
+public:
+	FGeometryCacheAbcFileSceneProxy(class UGeometryCacheAbcFileComponent* Component);
+};
+
+/** Track proxy that gets its data from a GeometryCacheAbcFileTrack */
+class FGeomCacheTrackAbcFileProxy : public FGeomCacheTrackProxy
+{
+public:
+	FGeomCacheTrackAbcFileProxy(ERHIFeatureLevel::Type InFeatureLevel)
+		: FGeomCacheTrackProxy(InFeatureLevel)
+	{}
+
+	//~ Begin FGeomCacheTrackProxy Interface
+	virtual bool UpdateMeshData(float Time, bool bLooping, int32& InOutMeshSampleIndex, FGeometryCacheMeshData& OutMeshData) override;
+	virtual bool GetMeshData(int32 SampleIndex, FGeometryCacheMeshData& OutMeshData) override;
+	virtual bool IsTopologyCompatible(int32 SampleIndexA, int32 SampleIndexB) override;
+	virtual const FVisibilitySample& GetVisibilitySample(float Time, const bool bLooping) const override;
+	virtual void FindSampleIndexesFromTime(float Time, bool bLooping, bool bIsPlayingBackwards, int32 &OutFrameIndex, int32 &OutNextFrameIndex, float &InterpolationFactor) override;
+	//~ End FGeomCacheTrackProxy Interface
+};
diff --git a/Engine/Plugins/Experimental/GeometryCacheAbcFile/Source/GeometryCacheAbcFile/Public/GeometryCacheTrackAbcFile.h b/Engine/Plugins/Experimental/GeometryCacheAbcFile/Source/GeometryCacheAbcFile/Public/GeometryCacheTrackAbcFile.h
new file mode 100644
--- /dev/null
+++ b/Engine/Plugins/Experimental/GeometryCacheAbcFile/Source/GeometryCacheAbcFile/Public/GeometryCacheTrackAbcFile.h
@@ -0,0 +1,43 @@
+// Copyright Epic Games, Inc. All Rights Reserved.
+
+#pragma once
+
+#include "CoreMinimal.h"
+#include "AbcFile.h"
+#include "GeometryCacheTrack.h"
+#include "GeometryCacheMeshData.h"
+
+#include "GeometryCacheTrackAbcFile.generated.h"
+
+/** GeometryCacheTrack for Alembic file querying */
+UCLASS(collapsecategories, hidecategories = Object, BlueprintType, config = Engine)
+class GEOMETRYCACHEABCFILE_API UGeometryCacheTrackAbcFile : public UGeometryCacheTrack
+{
+	GENERATED_BODY()
+
+	UGeometryCacheTrackAbcFile();
+
+public:
+
+	//~ Begin UGeometryCacheTrack Interface.
+	virtual const bool UpdateMeshData(const float Time, const bool bLooping, int32& InOutMeshSampleIndex, FGeometryCacheMeshData*& OutMeshData) override;
+	virtual const bool UpdateBoundsData(const float Time, const bool bLooping, const bool bIsPlayingBackward, int32& InOutBoundsSampleIndex, FBox& OutBounds) override;
+	virtual const FGeometryCacheTrackSampleInfo& GetSampleInfo(float Time, const bool bLooping) override;
+	//~ End UGeometryCacheTrack Interface.
+
+	bool SetSourceFile(const FString& FilePath);
+	const FString& GetSourceFile() const { return SourceFile; }
+
+	const int32 FindSampleIndexFromTime(const float Time, const bool bLooping) const;
+
+	int32 GetEndFrameIndex() const { return EndFrameIndex;  }
+
+	bool GetMeshData(int32 SampleIndex, FGeometryCacheMeshData& OutMeshData);
+
+private:
+	FGeometryCacheMeshData MeshData;
+	FGeometryCacheTrackSampleInfo SampleInfo;
+	TUniquePtr<FAbcFile> AbcFile;
+	FString SourceFile;
+	int32 EndFrameIndex;
+};
