Copied 11294374 from Dev-VirtualProduction.

Added Alembic settings to GeometryCacheAbcFileComponent.

diff --git a/Engine/Plugins/Experimental/AlembicImporter/Source/AlembicLibrary/Private/AbcImportSettingsCustomization.cpp b/Engine/Plugins/Experimental/AlembicImporter/Source/AlembicLibrary/Private/AbcImportSettingsCustomization.cpp
--- a/Engine/Plugins/Experimental/AlembicImporter/Source/AlembicLibrary/Private/AbcImportSettingsCustomization.cpp
+++ b/Engine/Plugins/Experimental/AlembicImporter/Source/AlembicLibrary/Private/AbcImportSettingsCustomization.cpp
@@ -77,6 +77,12 @@ void FAbcSamplingSettingsCustomization::CustomizeChildren(TSharedRef<IPropertyHa
 	{
 		TSharedRef<IPropertyHandle> ChildHandle = StructPropertyHandle->GetChildHandle(ChildIndex).ToSharedRef();
 
+		// Skip properties that are already marked hidden
+		if (ChildHandle->IsCustomized())
+		{
+			continue;
+		}
+
 		if (ChildHandle->GetProperty()->GetFName() == GET_MEMBER_NAME_CHECKED(FAbcSamplingSettings, SamplingType))
 		{
 			SamplingTypeHandle = ChildHandle; 
diff --git a/Engine/Plugins/Experimental/GeometryCacheAbcFile/Source/GeometryCacheAbcFile/GeometryCacheAbcFile.Build.cs b/Engine/Plugins/Experimental/GeometryCacheAbcFile/Source/GeometryCacheAbcFile/GeometryCacheAbcFile.Build.cs
--- a/Engine/Plugins/Experimental/GeometryCacheAbcFile/Source/GeometryCacheAbcFile/GeometryCacheAbcFile.Build.cs
+++ b/Engine/Plugins/Experimental/GeometryCacheAbcFile/Source/GeometryCacheAbcFile/GeometryCacheAbcFile.Build.cs
@@ -17,5 +17,13 @@ public class GeometryCacheAbcFile : ModuleRules
 				"RHI"
 			}
 		);
+
+		PrivateDependencyModuleNames.AddRange(
+			new string[] {
+				"PropertyEditor",
+				"Slate",
+				"SlateCore"
+			}
+		);
 	}
 }
diff --git a/Engine/Plugins/Experimental/GeometryCacheAbcFile/Source/GeometryCacheAbcFile/Private/GeometryCacheAbcFileComponent.cpp b/Engine/Plugins/Experimental/GeometryCacheAbcFile/Source/GeometryCacheAbcFile/Private/GeometryCacheAbcFileComponent.cpp
--- a/Engine/Plugins/Experimental/GeometryCacheAbcFile/Source/GeometryCacheAbcFile/Private/GeometryCacheAbcFileComponent.cpp
+++ b/Engine/Plugins/Experimental/GeometryCacheAbcFile/Source/GeometryCacheAbcFile/Private/GeometryCacheAbcFileComponent.cpp
@@ -7,6 +7,12 @@
 
 #define LOCTEXT_NAMESPACE "GeometryCacheAbcFileComponent"
 
+UGeometryCacheAbcFileComponent::UGeometryCacheAbcFileComponent()
+{
+	AbcSettings = NewObject<UAbcImportSettings>(this, TEXT("AbcSettings"));
+	AbcSettings->ImportType = EAlembicImportType::GeometryCache;
+}
+
 #if WITH_EDITOR
 void UGeometryCacheAbcFileComponent::PostEditChangeProperty(FPropertyChangedEvent& PropertyChangedEvent)
 {
@@ -19,10 +25,33 @@ void UGeometryCacheAbcFileComponent::PostEditChangeProperty(FPropertyChangedEven
 		InvalidateTrackSampleIndices();
 	}
 
-	Super::PostEditChangeProperty(PropertyChangedEvent);
+	UMeshComponent::PostEditChangeProperty(PropertyChangedEvent);
 }
 #endif
 
+void UGeometryCacheAbcFileComponent::ReloadAbcFile()
+{
+	if (!GeometryCache || GeometryCache->Tracks.Num() == 0 || AlembicFilePath.FilePath.IsEmpty())
+	{
+		return;
+	}
+
+	UGeometryCacheTrackAbcFile* AbcFileTrack = Cast<UGeometryCacheTrackAbcFile>(GeometryCache->Tracks[0]);
+
+	AbcSettings->SamplingSettings = SamplingSettings;
+	AbcSettings->MaterialSettings = MaterialSettings;
+	AbcSettings->ConversionSettings = ConversionSettings;
+
+	bool bIsValid = AbcFileTrack->SetSourceFile(AlembicFilePath.FilePath, AbcSettings);
+	if (bIsValid)
+	{
+		// Also store the number of frames in the cache
+		GeometryCache->SetFrameStartEnd(0, AbcFileTrack->GetEndFrameIndex());
+	}
+
+	MarkRenderStateDirty();
+}
+
 void UGeometryCacheAbcFileComponent::InitializeGeometryCache()
 {
 	UGeometryCacheTrackAbcFile* AbcFileTrack = nullptr;
@@ -41,15 +70,8 @@ void UGeometryCacheAbcFileComponent::InitializeGeometryCache()
 	// #ueent_todo: Should be able to clear the Alembic file
 	if (!AlembicFilePath.FilePath.IsEmpty() && (AlembicFilePath.FilePath != AbcFileTrack->GetSourceFile()))
 	{
-		bool bIsValid = AbcFileTrack->SetSourceFile(AlembicFilePath.FilePath);
-		if (bIsValid)
-		{
-			// Also store the number of frames in the cache
-			GeometryCache->SetFrameStartEnd(0, AbcFileTrack->GetEndFrameIndex());
-		}
+		ReloadAbcFile();
 	}
-
-	MarkRenderStateDirty();
 }
 
 void UGeometryCacheAbcFileComponent::PostLoad()
diff --git a/Engine/Plugins/Experimental/GeometryCacheAbcFile/Source/GeometryCacheAbcFile/Private/GeometryCacheAbcFileComponentCustomization.cpp b/Engine/Plugins/Experimental/GeometryCacheAbcFile/Source/GeometryCacheAbcFile/Private/GeometryCacheAbcFileComponentCustomization.cpp
new file mode 100644
--- /dev/null
+++ b/Engine/Plugins/Experimental/GeometryCacheAbcFile/Source/GeometryCacheAbcFile/Private/GeometryCacheAbcFileComponentCustomization.cpp
@@ -0,0 +1,80 @@
+// Copyright Epic Games, Inc. All Rights Reserved.
+
+#include "GeometryCacheAbcFileComponentCustomization.h"
+#include "AbcImportSettings.h"
+#include "DetailCategoryBuilder.h"
+#include "DetailLayoutBuilder.h"
+#include "DetailWidgetRow.h"
+#include "GeometryCacheAbcFileComponent.h"
+#include "Widgets/SBoxPanel.h"
+#include "Widgets/Input/SButton.h"
+
+#define LOCTEXT_NAMESPACE "GeometryCacheAbcFileComponentCustomization"
+
+TSharedRef<IDetailCustomization> FGeometryCacheAbcFileComponentCustomization::MakeInstance()
+{
+	return MakeShareable(new FGeometryCacheAbcFileComponentCustomization);
+}
+
+void FGeometryCacheAbcFileComponentCustomization::CustomizeDetails(IDetailLayoutBuilder& DetailBuilder)
+{
+	TArray<TWeakObjectPtr<UObject>> Objects;
+	DetailBuilder.GetObjectsBeingCustomized(Objects);
+	if (Objects.Num() != 1)
+	{
+		return;
+	}
+
+	GeometryCacheAbcFileComponent = Cast<UGeometryCacheAbcFileComponent>(Objects[0].Get());
+	if (!GeometryCacheAbcFileComponent.Get())
+	{
+		return;
+	}
+
+	// Hide properties that are irrelevant to the GeometryCacheAbcFileComponent
+	DetailBuilder.HideProperty("SamplingSettings.SamplingType");
+	DetailBuilder.HideProperty("SamplingSettings.FrameSteps");
+	DetailBuilder.HideProperty("SamplingSettings.TimeSteps");
+
+	// Add a button to reload the Alembic file and apply the settings
+	IDetailCategoryBuilder& Category = DetailBuilder.EditCategory("Alembic");
+
+	Category.AddCustomRow(LOCTEXT("AlembicSearchText", "Alembic"))
+	.ValueContent()
+	[
+		SNew(SHorizontalBox)
+		+SHorizontalBox::Slot()
+		.AutoWidth()
+		.Padding(2.0f, 0.0f)
+		.VAlign(VAlign_Center)
+		.HAlign(HAlign_Left)
+		[
+			SNew(SButton)
+			.Text(LOCTEXT("ReloadAbcFile", "Reload"))
+			.ToolTipText(LOCTEXT("ReloadAbcFile_Tooltip", "Reload the Alembic file and apply the settings"))
+			.OnClicked(this, &FGeometryCacheAbcFileComponentCustomization::ReloadAbcFile)
+			.IsEnabled(this, &FGeometryCacheAbcFileComponentCustomization::IsReloadEnabled)
+		]
+	];
+}
+
+FReply FGeometryCacheAbcFileComponentCustomization::ReloadAbcFile()
+{
+	if (GeometryCacheAbcFileComponent.Get())
+	{
+		GeometryCacheAbcFileComponent->ReloadAbcFile();
+	}
+
+	return FReply::Handled();
+}
+
+bool FGeometryCacheAbcFileComponentCustomization::IsReloadEnabled() const
+{
+	if (GeometryCacheAbcFileComponent.Get())
+	{
+		return !GeometryCacheAbcFileComponent->AlembicFilePath.FilePath.IsEmpty();
+	}
+	return false;
+}
+
+#undef LOCTEXT_NAMESPACE
diff --git a/Engine/Plugins/Experimental/GeometryCacheAbcFile/Source/GeometryCacheAbcFile/Private/GeometryCacheAbcFileComponentCustomization.h b/Engine/Plugins/Experimental/GeometryCacheAbcFile/Source/GeometryCacheAbcFile/Private/GeometryCacheAbcFileComponentCustomization.h
new file mode 100644
--- /dev/null
+++ b/Engine/Plugins/Experimental/GeometryCacheAbcFile/Source/GeometryCacheAbcFile/Private/GeometryCacheAbcFileComponentCustomization.h
@@ -0,0 +1,25 @@
+// Copyright Epic Games, Inc. All Rights Reserved.
+
+#pragma once
+
+#include "CoreMinimal.h"
+#include "Input/Reply.h"
+#include "IDetailCustomization.h"
+
+class IDetailLayoutBuilder;
+
+class FGeometryCacheAbcFileComponentCustomization : public IDetailCustomization
+{
+public:
+	static TSharedRef<IDetailCustomization> MakeInstance();
+
+	// IDetailCustomization interface
+	virtual void CustomizeDetails(IDetailLayoutBuilder& DetailLayout) override;
+	// End of IDetailCustomization interface
+
+protected:
+	FReply ReloadAbcFile();
+	bool IsReloadEnabled() const;
+
+	TWeakObjectPtr<class UGeometryCacheAbcFileComponent> GeometryCacheAbcFileComponent;
+};
diff --git a/Engine/Plugins/Experimental/GeometryCacheAbcFile/Source/GeometryCacheAbcFile/Private/GeometryCacheAbcFileModule.cpp b/Engine/Plugins/Experimental/GeometryCacheAbcFile/Source/GeometryCacheAbcFile/Private/GeometryCacheAbcFileModule.cpp
--- a/Engine/Plugins/Experimental/GeometryCacheAbcFile/Source/GeometryCacheAbcFile/Private/GeometryCacheAbcFileModule.cpp
+++ b/Engine/Plugins/Experimental/GeometryCacheAbcFile/Source/GeometryCacheAbcFile/Private/GeometryCacheAbcFileModule.cpp
@@ -3,8 +3,29 @@
 #include "Modules/ModuleInterface.h"
 #include "Modules/ModuleManager.h"
 
+#include "GeometryCacheAbcFileComponent.h"
+#include "GeometryCacheAbcFileComponentCustomization.h"
+#include "PropertyEditorModule.h"
+
 class FGeometryCacheAbcFileModule : public IModuleInterface
 {
+	virtual void StartupModule() override
+	{
+		FPropertyEditorModule& PropertyModule = FModuleManager::LoadModuleChecked<FPropertyEditorModule>("PropertyEditor");
+		PropertyModule.RegisterCustomClassLayout(UGeometryCacheAbcFileComponent::StaticClass()->GetFName(), FOnGetDetailCustomizationInstance::CreateStatic(&FGeometryCacheAbcFileComponentCustomization::MakeInstance));
+	}
+
+	virtual void ShutdownModule() override
+	{
+		if (UObjectInitialized() && !IsEngineExitRequested())
+		{
+			FPropertyEditorModule* PropertyModule = FModuleManager::GetModulePtr<FPropertyEditorModule>("PropertyEditor");
+			if (PropertyModule)
+			{
+				PropertyModule->UnregisterCustomClassLayout(UGeometryCacheAbcFileComponent::StaticClass()->GetFName());
+			}
+		}
+	}
 };
 
 IMPLEMENT_MODULE(FGeometryCacheAbcFileModule, GeometryCacheAbcFile);
\ No newline at end of file
diff --git a/Engine/Plugins/Experimental/GeometryCacheAbcFile/Source/GeometryCacheAbcFile/Private/GeometryCacheTrackAbcFile.cpp b/Engine/Plugins/Experimental/GeometryCacheAbcFile/Source/GeometryCacheAbcFile/Private/GeometryCacheTrackAbcFile.cpp
--- a/Engine/Plugins/Experimental/GeometryCacheAbcFile/Source/GeometryCacheAbcFile/Private/GeometryCacheTrackAbcFile.cpp
+++ b/Engine/Plugins/Experimental/GeometryCacheAbcFile/Source/GeometryCacheAbcFile/Private/GeometryCacheTrackAbcFile.cpp
@@ -6,6 +6,7 @@
 #include "AbcUtilities.h"
 #include "AbcImportSettings.h"
 #include "GeometryCacheHelpers.h"
+#include "AbcImportSettings.h"
 
 UGeometryCacheTrackAbcFile::UGeometryCacheTrackAbcFile()
 : EndFrameIndex(0)
@@ -44,7 +45,7 @@ const bool UGeometryCacheTrackAbcFile::UpdateBoundsData(const float Time, const
 	return false;
 }
 
-bool UGeometryCacheTrackAbcFile::SetSourceFile(const FString& FilePath)
+bool UGeometryCacheTrackAbcFile::SetSourceFile(const FString& FilePath, UAbcImportSettings* AbcSettings)
 {
 	if (!FilePath.IsEmpty())
 	{
@@ -57,14 +58,14 @@ bool UGeometryCacheTrackAbcFile::SetSourceFile(const FString& FilePath)
 			return false;
 		}
 
-		// #ueent_todo: Expose some conversion settings in the component to use with the AbcFile
-		UAbcImportSettings* ImportSettings = DuplicateObject(GetMutableDefault<UAbcImportSettings>(), GetTransientPackage());
-
+		// Set the end frame from the Abc if none is specified in the settings
 		EndFrameIndex = FMath::Max(AbcFile->GetMaxFrameIndex() - 1, 1);
-		ImportSettings->SamplingSettings.FrameEnd = EndFrameIndex;
-		ImportSettings->ImportType = EAlembicImportType::GeometryCache;
+		if (AbcSettings->SamplingSettings.FrameEnd == 0)
+		{
+			AbcSettings->SamplingSettings.FrameEnd = EndFrameIndex;
+		}
 
-		Result = AbcFile->Import(ImportSettings);
+		Result = AbcFile->Import(AbcSettings);
 
 		if (Result != EAbcImportError::AbcImportError_NoError)
 		{
diff --git a/Engine/Plugins/Experimental/GeometryCacheAbcFile/Source/GeometryCacheAbcFile/Public/GeometryCacheAbcFileComponent.h b/Engine/Plugins/Experimental/GeometryCacheAbcFile/Source/GeometryCacheAbcFile/Public/GeometryCacheAbcFileComponent.h
--- a/Engine/Plugins/Experimental/GeometryCacheAbcFile/Source/GeometryCacheAbcFile/Public/GeometryCacheAbcFileComponent.h
+++ b/Engine/Plugins/Experimental/GeometryCacheAbcFile/Source/GeometryCacheAbcFile/Public/GeometryCacheAbcFileComponent.h
@@ -3,6 +3,7 @@
 #pragma once
 
 #include "CoreMinimal.h"
+#include "AbcImportSettings.h"
 #include "GeometryCacheComponent.h"
 
 #include "GeometryCacheAbcFileComponent.generated.h"
@@ -15,11 +16,22 @@ class GEOMETRYCACHEABCFILE_API UGeometryCacheAbcFileComponent : public UGeometry
 {
 	GENERATED_BODY()
 
+	UGeometryCacheAbcFileComponent();
+
 public:
 
 	UPROPERTY(EditAnywhere, Category = "Alembic", meta = (FilePathFilter = "Alembic files (*.abc)|*.abc"))
 	FFilePath AlembicFilePath;
 
+	UPROPERTY(EditAnywhere, BlueprintReadWrite, Category = "Alembic")
+	FAbcSamplingSettings SamplingSettings;
+
+	UPROPERTY(EditAnywhere, BlueprintReadWrite, Category = "Alembic")
+	FAbcMaterialSettings MaterialSettings;
+
+	UPROPERTY(EditAnywhere, BlueprintReadWrite, Category = "Alembic")
+	FAbcConversionSettings ConversionSettings;
+
 public:
 #if WITH_EDITOR
 	virtual void PostEditChangeProperty(FPropertyChangedEvent& PropertyChangedEvent) override;
@@ -33,6 +45,10 @@ public:
 	virtual FPrimitiveSceneProxy* CreateSceneProxy() override;
 	//~ End UPrimitiveComponent Interface.
 
+	void ReloadAbcFile();
+
 protected:
 	void InitializeGeometryCache();
+
+	UAbcImportSettings* AbcSettings;
 };
diff --git a/Engine/Plugins/Experimental/GeometryCacheAbcFile/Source/GeometryCacheAbcFile/Public/GeometryCacheTrackAbcFile.h b/Engine/Plugins/Experimental/GeometryCacheAbcFile/Source/GeometryCacheAbcFile/Public/GeometryCacheTrackAbcFile.h
--- a/Engine/Plugins/Experimental/GeometryCacheAbcFile/Source/GeometryCacheAbcFile/Public/GeometryCacheTrackAbcFile.h
+++ b/Engine/Plugins/Experimental/GeometryCacheAbcFile/Source/GeometryCacheAbcFile/Public/GeometryCacheTrackAbcFile.h
@@ -25,7 +25,7 @@ public:
 	virtual const FGeometryCacheTrackSampleInfo& GetSampleInfo(float Time, const bool bLooping) override;
 	//~ End UGeometryCacheTrack Interface.
 
-	bool SetSourceFile(const FString& FilePath);
+	bool SetSourceFile(const FString& FilePath, class UAbcImportSettings* AbcSettings);
 	const FString& GetSourceFile() const { return SourceFile; }
 
 	const int32 FindSampleIndexFromTime(const float Time, const bool bLooping) const;
