Copied 11303546 from Dev-VirtualProduction.

Added actor factory for GeometryCacheAbcFile.

diff --git a/Engine/Plugins/Experimental/GeometryCacheAbcFile/Source/GeometryCacheAbcFile/GeometryCacheAbcFile.Build.cs b/Engine/Plugins/Experimental/GeometryCacheAbcFile/Source/GeometryCacheAbcFile/GeometryCacheAbcFile.Build.cs
--- a/Engine/Plugins/Experimental/GeometryCacheAbcFile/Source/GeometryCacheAbcFile/GeometryCacheAbcFile.Build.cs
+++ b/Engine/Plugins/Experimental/GeometryCacheAbcFile/Source/GeometryCacheAbcFile/GeometryCacheAbcFile.Build.cs
@@ -22,8 +22,9 @@ public class GeometryCacheAbcFile : ModuleRules
 			new string[] {
 				"PropertyEditor",
 				"Slate",
-				"SlateCore"
-			}
-		);
+				"SlateCore",
+                "UnrealEd"
+            }
+        );
 	}
 }
diff --git a/Engine/Plugins/Experimental/GeometryCacheAbcFile/Source/GeometryCacheAbcFile/Private/ActorFactoryGeometryCacheAbcFile.cpp b/Engine/Plugins/Experimental/GeometryCacheAbcFile/Source/GeometryCacheAbcFile/Private/ActorFactoryGeometryCacheAbcFile.cpp
new file mode 100644
--- /dev/null
+++ b/Engine/Plugins/Experimental/GeometryCacheAbcFile/Source/GeometryCacheAbcFile/Private/ActorFactoryGeometryCacheAbcFile.cpp
@@ -0,0 +1,11 @@
+// Copyright Epic Games, Inc. All Rights Reserved.
+
+#include "ActorFactoryGeometryCacheAbcFile.h"
+#include "GeometryCacheAbcFileActor.h"
+
+UActorFactoryGeometryCacheAbcFile::UActorFactoryGeometryCacheAbcFile()
+{
+	DisplayName = FText::FromString( "Geometry Cache Abc File" );
+	NewActorClass = AGeometryCacheAbcFileActor::StaticClass();
+	bUseSurfaceOrientation = true;
+}
diff --git a/Engine/Plugins/Experimental/GeometryCacheAbcFile/Source/GeometryCacheAbcFile/Private/ActorFactoryGeometryCacheAbcFile.h b/Engine/Plugins/Experimental/GeometryCacheAbcFile/Source/GeometryCacheAbcFile/Private/ActorFactoryGeometryCacheAbcFile.h
new file mode 100644
--- /dev/null
+++ b/Engine/Plugins/Experimental/GeometryCacheAbcFile/Source/GeometryCacheAbcFile/Private/ActorFactoryGeometryCacheAbcFile.h
@@ -0,0 +1,17 @@
+// Copyright Epic Games, Inc. All Rights Reserved.
+
+#pragma once
+
+#include "CoreMinimal.h"
+#include "UObject/ObjectMacros.h"
+#include "ActorFactories/ActorFactory.h"
+
+#include "ActorFactoryGeometryCacheAbcFile.generated.h"
+
+UCLASS()
+class UActorFactoryGeometryCacheAbcFile : public UActorFactory
+{
+	GENERATED_BODY()
+
+	UActorFactoryGeometryCacheAbcFile();
+};
diff --git a/Engine/Plugins/Experimental/GeometryCacheAbcFile/Source/GeometryCacheAbcFile/Private/GeometryCacheAbcFileActor.cpp b/Engine/Plugins/Experimental/GeometryCacheAbcFile/Source/GeometryCacheAbcFile/Private/GeometryCacheAbcFileActor.cpp
new file mode 100644
--- /dev/null
+++ b/Engine/Plugins/Experimental/GeometryCacheAbcFile/Source/GeometryCacheAbcFile/Private/GeometryCacheAbcFileActor.cpp
@@ -0,0 +1,12 @@
+// Copyright Epic Games, Inc. All Rights Reserved.
+
+#include "GeometryCacheAbcFileActor.h"
+#include "GeometryCache.h"
+#include "GeometryCacheAbcFileComponent.h"
+
+AGeometryCacheAbcFileActor::AGeometryCacheAbcFileActor(const FObjectInitializer& ObjectInitializer)
+: Super(ObjectInitializer)
+{
+	GeometryCacheAbcFileComponent = CreateDefaultSubobject<UGeometryCacheAbcFileComponent>(TEXT("GeometryCacheAbcFileComponent"));
+	RootComponent = GeometryCacheAbcFileComponent;
+}
diff --git a/Engine/Plugins/Experimental/GeometryCacheAbcFile/Source/GeometryCacheAbcFile/Private/GeometryCacheAbcFileActor.h b/Engine/Plugins/Experimental/GeometryCacheAbcFile/Source/GeometryCacheAbcFile/Private/GeometryCacheAbcFileActor.h
new file mode 100644
--- /dev/null
+++ b/Engine/Plugins/Experimental/GeometryCacheAbcFile/Source/GeometryCacheAbcFile/Private/GeometryCacheAbcFileActor.h
@@ -0,0 +1,22 @@
+// Copyright Epic Games, Inc. All Rights Reserved.
+
+#pragma once
+
+#include "CoreMinimal.h"
+#include "UObject/ObjectMacros.h"
+#include "GameFramework/Actor.h"
+
+#include "GeometryCacheAbcFileActor.generated.h"
+
+class UGeometryCacheAbcFileComponent;
+
+/** GeometryCacheAbcFile actor serves as a placeable actor for GeometryCache loading from an Alembic file */
+UCLASS(ComponentWrapperClass)
+class AGeometryCacheAbcFileActor : public AActor
+{
+	GENERATED_UCLASS_BODY()
+
+private:
+	UPROPERTY(Category = GeometryCacheAbcFileActor, VisibleAnywhere)
+	UGeometryCacheAbcFileComponent* GeometryCacheAbcFileComponent;
+};
