Copied 11307851 from Dev-VirtualProduction.

Fixed crash in PIE with an empty GeometryCacheAbcFile actor in the level.
Fixed crash when clearing the Alembic File Path on the GeometryCacheAbcFile component.

diff --git a/Engine/Plugins/Experimental/GeometryCacheAbcFile/Source/GeometryCacheAbcFile/Private/GeometryCacheAbcFileComponent.cpp b/Engine/Plugins/Experimental/GeometryCacheAbcFile/Source/GeometryCacheAbcFile/Private/GeometryCacheAbcFileComponent.cpp
--- a/Engine/Plugins/Experimental/GeometryCacheAbcFile/Source/GeometryCacheAbcFile/Private/GeometryCacheAbcFileComponent.cpp
+++ b/Engine/Plugins/Experimental/GeometryCacheAbcFile/Source/GeometryCacheAbcFile/Private/GeometryCacheAbcFileComponent.cpp
@@ -10,7 +10,6 @@
 UGeometryCacheAbcFileComponent::UGeometryCacheAbcFileComponent()
 {
 	AbcSettings = NewObject<UAbcImportSettings>(this, TEXT("AbcSettings"));
-	AbcSettings->ImportType = EAlembicImportType::GeometryCache;
 }
 
 #if WITH_EDITOR
@@ -21,7 +20,22 @@ void UGeometryCacheAbcFileComponent::PostEditChangeProperty(FPropertyChangedEven
 
 	if (PropertyName == GET_MEMBER_NAME_CHECKED(UGeometryCacheAbcFileComponent, AlembicFilePath))
 	{
-		InitializeGeometryCache();
+		if (!AlembicFilePath.FilePath.IsEmpty())
+		{
+			InitializeGeometryCache();
+		}
+		else
+		{
+			if (GeometryCache && GeometryCache->Tracks.Num() != 0)
+			{
+				// Release the AbcFile resources
+				UGeometryCacheTrackAbcFile* AbcFileTrack = Cast<UGeometryCacheTrackAbcFile>(GeometryCache->Tracks[0]);
+				AbcFileTrack->SetSourceFile(FString(), nullptr);
+			}
+
+			GeometryCache = nullptr;
+			MarkRenderStateDirty();
+		}
 		InvalidateTrackSampleIndices();
 	}
 
@@ -38,6 +52,7 @@ void UGeometryCacheAbcFileComponent::ReloadAbcFile()
 
 	UGeometryCacheTrackAbcFile* AbcFileTrack = Cast<UGeometryCacheTrackAbcFile>(GeometryCache->Tracks[0]);
 
+	AbcSettings->ImportType = EAlembicImportType::GeometryCache;
 	AbcSettings->SamplingSettings = SamplingSettings;
 	AbcSettings->MaterialSettings = MaterialSettings;
 	AbcSettings->ConversionSettings = ConversionSettings;
@@ -54,23 +69,25 @@ void UGeometryCacheAbcFileComponent::ReloadAbcFile()
 
 void UGeometryCacheAbcFileComponent::InitializeGeometryCache()
 {
-	UGeometryCacheTrackAbcFile* AbcFileTrack = nullptr;
-	if (!GeometryCache)
-	{
-		// Transient GeometryCache for use in current session
-		GeometryCache = NewObject<UGeometryCache>();
-		AbcFileTrack = NewObject<UGeometryCacheTrackAbcFile>(GeometryCache);
-		GeometryCache->AddTrack(AbcFileTrack);
-	}
-	else
-	{
-		AbcFileTrack = Cast<UGeometryCacheTrackAbcFile>(GeometryCache->Tracks[0]);
-	}
-
-	// #ueent_todo: Should be able to clear the Alembic file
-	if (!AlembicFilePath.FilePath.IsEmpty() && (AlembicFilePath.FilePath != AbcFileTrack->GetSourceFile()))
+	if (!AlembicFilePath.FilePath.IsEmpty())
 	{
-		ReloadAbcFile();
+		UGeometryCacheTrackAbcFile* AbcFileTrack = nullptr;
+		if (!GeometryCache)
+		{
+			// Transient GeometryCache for use in current session
+			GeometryCache = NewObject<UGeometryCache>();
+			AbcFileTrack = NewObject<UGeometryCacheTrackAbcFile>(GeometryCache);
+			GeometryCache->AddTrack(AbcFileTrack);
+		}
+		else
+		{
+			AbcFileTrack = Cast<UGeometryCacheTrackAbcFile>(GeometryCache->Tracks[0]);
+		}
+
+		if (AlembicFilePath.FilePath != AbcFileTrack->GetSourceFile())
+		{
+			ReloadAbcFile();
+		}
 	}
 }
 
diff --git a/Engine/Plugins/Experimental/GeometryCacheAbcFile/Source/GeometryCacheAbcFile/Private/GeometryCacheTrackAbcFile.cpp b/Engine/Plugins/Experimental/GeometryCacheAbcFile/Source/GeometryCacheAbcFile/Private/GeometryCacheTrackAbcFile.cpp
--- a/Engine/Plugins/Experimental/GeometryCacheAbcFile/Source/GeometryCacheAbcFile/Private/GeometryCacheTrackAbcFile.cpp
+++ b/Engine/Plugins/Experimental/GeometryCacheAbcFile/Source/GeometryCacheAbcFile/Private/GeometryCacheTrackAbcFile.cpp
@@ -13,6 +13,15 @@ UGeometryCacheTrackAbcFile::UGeometryCacheTrackAbcFile()
 {
 }
 
+const bool UGeometryCacheTrackAbcFile::UpdateMatrixData(const float Time, const bool bLooping, int32& InOutMatrixSampleIndex, FMatrix& OutWorldMatrix)
+{
+	if (AbcFile)
+	{
+		return Super::UpdateMatrixData(Time, bLooping, InOutMatrixSampleIndex, OutWorldMatrix);
+	}
+	return false;
+}
+
 const bool UGeometryCacheTrackAbcFile::UpdateMeshData(const float Time, const bool bLooping, int32& InOutMeshSampleIndex, FGeometryCacheMeshData*& OutMeshData)
 {
 	const int32 SampleIndex = FindSampleIndexFromTime(Time, bLooping);
@@ -73,6 +82,9 @@ bool UGeometryCacheTrackAbcFile::SetSourceFile(const FString& FilePath, UAbcImpo
 			return false;
 		}
 
+		MatrixSamples.Reset();
+		MatrixSampleTimes.Reset();
+
 		TArray<FMatrix> Mats;
 		Mats.Add(FMatrix::Identity);
 		Mats.Add(FMatrix::Identity);
@@ -137,7 +149,7 @@ const FGeometryCacheTrackSampleInfo& UGeometryCacheTrackAbcFile::GetSampleInfo(f
 
 bool UGeometryCacheTrackAbcFile::GetMeshData(int32 SampleIndex, FGeometryCacheMeshData& OutMeshData)
 {
-	if (AbcFile.IsValid())
+	if (AbcFile)
 	{
 		// #ueent_todo: Implement optimized Alembic querying
 		FAbcUtilities::GetFrameMeshData(*AbcFile, SampleIndex, OutMeshData);
diff --git a/Engine/Plugins/Experimental/GeometryCacheAbcFile/Source/GeometryCacheAbcFile/Public/GeometryCacheTrackAbcFile.h b/Engine/Plugins/Experimental/GeometryCacheAbcFile/Source/GeometryCacheAbcFile/Public/GeometryCacheTrackAbcFile.h
--- a/Engine/Plugins/Experimental/GeometryCacheAbcFile/Source/GeometryCacheAbcFile/Public/GeometryCacheTrackAbcFile.h
+++ b/Engine/Plugins/Experimental/GeometryCacheAbcFile/Source/GeometryCacheAbcFile/Public/GeometryCacheTrackAbcFile.h
@@ -20,6 +20,7 @@ class GEOMETRYCACHEABCFILE_API UGeometryCacheTrackAbcFile : public UGeometryCach
 public:
 
 	//~ Begin UGeometryCacheTrack Interface.
+	virtual const bool UpdateMatrixData(const float Time, const bool bLooping, int32& InOutMatrixSampleIndex, FMatrix& OutWorldMatrix) override;
 	virtual const bool UpdateMeshData(const float Time, const bool bLooping, int32& InOutMeshSampleIndex, FGeometryCacheMeshData*& OutMeshData) override;
 	virtual const bool UpdateBoundsData(const float Time, const bool bLooping, const bool bIsPlayingBackward, int32& InOutBoundsSampleIndex, FBox& OutBounds) override;
 	virtual const FGeometryCacheTrackSampleInfo& GetSampleInfo(float Time, const bool bLooping) override;
