Copied 11392843 from Dev-VirtualProduction.

Fixed crash with GeometryCacheAbcFile when trying to load files that don't have any visible geometry or that are made of elements other than triangles and quads.

diff --git a/Engine/Plugins/Experimental/AlembicImporter/Source/AlembicLibrary/Private/AbcImportLogger.cpp b/Engine/Plugins/Experimental/AlembicImporter/Source/AlembicLibrary/Private/AbcImportLogger.cpp
--- a/Engine/Plugins/Experimental/AlembicImporter/Source/AlembicLibrary/Private/AbcImportLogger.cpp
+++ b/Engine/Plugins/Experimental/AlembicImporter/Source/AlembicLibrary/Private/AbcImportLogger.cpp
@@ -26,3 +26,8 @@ void FAbcImportLogger::OutputMessages(const FString& PageName)
 
 	MessageLog.Open();
 }
+
+TArray<TSharedRef<FTokenizedMessage>> FAbcImportLogger::RetrieveMessages()
+{	
+	return MoveTemp(TokenizedErrorMessages);
+}
diff --git a/Engine/Plugins/Experimental/AlembicImporter/Source/AlembicLibrary/Public/AbcImportLogger.h b/Engine/Plugins/Experimental/AlembicImporter/Source/AlembicLibrary/Public/AbcImportLogger.h
--- a/Engine/Plugins/Experimental/AlembicImporter/Source/AlembicLibrary/Public/AbcImportLogger.h
+++ b/Engine/Plugins/Experimental/AlembicImporter/Source/AlembicLibrary/Public/AbcImportLogger.h
@@ -14,6 +14,8 @@ public:
 	static void AddImportMessage(const TSharedRef<FTokenizedMessage> Message);
 	/** Outputs the messages to a new named page in the message log */
 	ALEMBICLIBRARY_API static void OutputMessages(const FString& PageName);
+	/** Returns the messages and flushes them from the logger */
+	ALEMBICLIBRARY_API static TArray<TSharedRef<FTokenizedMessage>> RetrieveMessages();
 private:
 	/** Error messages **/
 	static TArray<TSharedRef<FTokenizedMessage>> TokenizedErrorMessages;
diff --git a/Engine/Plugins/Experimental/GeometryCache/Source/GeometryCache/Private/GeometryCacheSceneProxy.cpp b/Engine/Plugins/Experimental/GeometryCache/Source/GeometryCache/Private/GeometryCacheSceneProxy.cpp
--- a/Engine/Plugins/Experimental/GeometryCache/Source/GeometryCache/Private/GeometryCacheSceneProxy.cpp
+++ b/Engine/Plugins/Experimental/GeometryCache/Source/GeometryCache/Private/GeometryCacheSceneProxy.cpp
@@ -65,7 +65,7 @@ FGeometryCacheSceneProxy::FGeometryCacheSceneProxy(UGeometryCacheComponent* Comp
 
 	// Copy each section
 	const int32 NumTracks = Component->TrackSections.Num();
-	Tracks.AddZeroed(NumTracks);
+	Tracks.Reserve(NumTracks);
 	for (int32 TrackIdx = 0; TrackIdx < NumTracks; TrackIdx++)
 	{
 		FTrackRenderData& SrcSection = Component->TrackSections[TrackIdx];
@@ -73,6 +73,7 @@ FGeometryCacheSceneProxy::FGeometryCacheSceneProxy(UGeometryCacheComponent* Comp
 
 		const FGeometryCacheTrackSampleInfo& SampleInfo = CurrentTrack->GetSampleInfo(Time, bLooping);
 
+		// Add track only if it has (visible) geometry
 		if (SampleInfo.NumVertices > 0)
 		{
 			FGeomCacheTrackProxy* NewSection = CreateTrackProxy();
@@ -135,7 +136,7 @@ FGeometryCacheSceneProxy::FGeometryCacheSceneProxy(UGeometryCacheComponent* Comp
 			}
 
 			// Save ref to new section
-			Tracks[TrackIdx] = NewSection;
+			Tracks.Add(NewSection);
 		}
 	}
 
diff --git a/Engine/Plugins/Experimental/GeometryCacheAbcFile/Source/GeometryCacheAbcFile/Private/GeometryCacheAbcFileComponent.cpp b/Engine/Plugins/Experimental/GeometryCacheAbcFile/Source/GeometryCacheAbcFile/Private/GeometryCacheAbcFileComponent.cpp
--- a/Engine/Plugins/Experimental/GeometryCacheAbcFile/Source/GeometryCacheAbcFile/Private/GeometryCacheAbcFileComponent.cpp
+++ b/Engine/Plugins/Experimental/GeometryCacheAbcFile/Source/GeometryCacheAbcFile/Private/GeometryCacheAbcFileComponent.cpp
@@ -63,6 +63,11 @@ void UGeometryCacheAbcFileComponent::ReloadAbcFile()
 		// Also store the number of frames in the cache
 		GeometryCache->SetFrameStartEnd(0, AbcFileTrack->GetEndFrameIndex());
 	}
+	else
+	{
+		GeometryCache = nullptr;
+		AlembicFilePath.FilePath.Empty();
+	}
 
 	MarkRenderStateDirty();
 }
diff --git a/Engine/Plugins/Experimental/GeometryCacheAbcFile/Source/GeometryCacheAbcFile/Private/GeometryCacheTrackAbcFile.cpp b/Engine/Plugins/Experimental/GeometryCacheAbcFile/Source/GeometryCacheAbcFile/Private/GeometryCacheTrackAbcFile.cpp
--- a/Engine/Plugins/Experimental/GeometryCacheAbcFile/Source/GeometryCacheAbcFile/Private/GeometryCacheTrackAbcFile.cpp
+++ b/Engine/Plugins/Experimental/GeometryCacheAbcFile/Source/GeometryCacheAbcFile/Private/GeometryCacheTrackAbcFile.cpp
@@ -3,10 +3,18 @@
 #include "GeometryCacheTrackAbcFile.h"
 
 #include "AbcImporter.h"
-#include "AbcUtilities.h"
+#include "AbcImportLogger.h"
 #include "AbcImportSettings.h"
+#include "AbcUtilities.h"
+#include "Framework/Notifications/NotificationManager.h"
 #include "GeometryCacheHelpers.h"
-#include "AbcImportSettings.h"
+#include "Logging/LogCategory.h"
+#include "Logging/LogMacros.h"
+#include "Widgets/Notifications/SNotificationList.h"
+
+DEFINE_LOG_CATEGORY_STATIC(LogGeometryCacheAbcFile, Log, All);
+
+#define LOCTEXT_NAMESPACE "GeometryCacheTrackAbcFile"
 
 UGeometryCacheTrackAbcFile::UGeometryCacheTrackAbcFile()
 : EndFrameIndex(0)
@@ -54,6 +62,31 @@ const bool UGeometryCacheTrackAbcFile::UpdateBoundsData(const float Time, const
 	return false;
 }
 
+void UGeometryCacheTrackAbcFile::Reset()
+{
+	AbcFile.Reset();
+
+	EndFrameIndex = 0;
+	Duration = 0.f;
+
+	MatrixSamples.Reset();
+	MatrixSampleTimes.Reset();
+
+	MeshData = FGeometryCacheMeshData();
+	MeshData.BoundingBox = FBox(ForceInit);
+}
+
+void UGeometryCacheTrackAbcFile::ShowNotification(const FText& Text)
+{
+	FNotificationInfo Info(Text);
+	Info.bFireAndForget = true;
+	Info.bUseLargeFont = false;
+	Info.FadeOutDuration = 3.0f;
+	Info.ExpireDuration = 7.0f;
+
+	FSlateNotificationManager::Get().AddNotification(Info);
+}
+
 bool UGeometryCacheTrackAbcFile::SetSourceFile(const FString& FilePath, UAbcImportSettings* AbcSettings)
 {
 	if (!FilePath.IsEmpty())
@@ -61,9 +94,24 @@ bool UGeometryCacheTrackAbcFile::SetSourceFile(const FString& FilePath, UAbcImpo
 		AbcFile = MakeUnique<FAbcFile>(FilePath);
 		EAbcImportError Result = AbcFile->Open();
 
+		const FString Filename = FPaths::GetCleanFilename(FilePath);
+
 		if (Result != EAbcImportError::AbcImportError_NoError)
 		{
-			AbcFile.Reset();
+			Reset();
+
+			FText FailureMessage = LOCTEXT("OpenFailureReason_Unknown", "Unknown open failure");
+			switch (Result)
+			{
+			case EAbcImportError::AbcImportError_InvalidArchive:
+				FailureMessage = LOCTEXT("OpenFailureReason_InvalidArchive", "Not a valid Alembic file");
+				break;
+			case EAbcImportError::AbcImportError_NoValidTopObject:
+				FailureMessage = LOCTEXT("OpenFailureReason_InvalidRoot", "Alembic file has no valid root node");
+				break;
+			}
+			UE_LOG(LogGeometryCacheAbcFile, Warning, TEXT("Failed to open %s: %s"), *Filename, *FailureMessage.ToString());
+
 			return false;
 		}
 
@@ -78,7 +126,18 @@ bool UGeometryCacheTrackAbcFile::SetSourceFile(const FString& FilePath, UAbcImpo
 
 		if (Result != EAbcImportError::AbcImportError_NoError)
 		{
-			AbcFile.Reset();
+			Reset();
+
+			FText FailureMessage = LOCTEXT("LoadFailureReason_Unknown", "Unknown load failure");
+			TArray<TSharedRef<FTokenizedMessage>> Messages = FAbcImportLogger::RetrieveMessages();
+			if (Messages.Num() > 0)
+			{
+				FailureMessage = Messages[0]->ToText();
+			}
+			UE_LOG(LogGeometryCacheAbcFile, Warning, TEXT("Failed to load %s: %s"), *Filename, *FailureMessage.ToString());
+
+			ShowNotification(FText::Format(LOCTEXT("LoadErrorNotification", "{0} could not be loaded. See Output Log for details."), FText::FromString(Filename)));
+
 			return false;
 		}
 
@@ -98,16 +157,16 @@ bool UGeometryCacheTrackAbcFile::SetSourceFile(const FString& FilePath, UAbcImpo
 
 		// Fill out the MeshData with the sample from time 0
 		GetMeshData(0, MeshData);
+
+		if (MeshData.Positions.Num() == 0)
+		{
+			// This could happen if the Alembic has geometry but they are set as invisible in the source
+			ShowNotification(FText::Format(LOCTEXT("NoVisibleGeometry", "Warning: {0} has no visible geometry."), FText::FromString(Filename)));
+		}
 	}
 	else
 	{
-		AbcFile.Reset();
-
-		EndFrameIndex = 0;
-		Duration = 0.f;
-
-		MatrixSamples.Empty();
-		MatrixSampleTimes.Empty();
+		Reset();
 	}
 	SourceFile = FilePath;
 	return true;
@@ -157,3 +216,5 @@ bool UGeometryCacheTrackAbcFile::GetMeshData(int32 SampleIndex, FGeometryCacheMe
 	}
 	return false;
 }
+
+#undef LOCTEXT_NAMESPACE
\ No newline at end of file
diff --git a/Engine/Plugins/Experimental/GeometryCacheAbcFile/Source/GeometryCacheAbcFile/Public/GeometryCacheTrackAbcFile.h b/Engine/Plugins/Experimental/GeometryCacheAbcFile/Source/GeometryCacheAbcFile/Public/GeometryCacheTrackAbcFile.h
--- a/Engine/Plugins/Experimental/GeometryCacheAbcFile/Source/GeometryCacheAbcFile/Public/GeometryCacheTrackAbcFile.h
+++ b/Engine/Plugins/Experimental/GeometryCacheAbcFile/Source/GeometryCacheAbcFile/Public/GeometryCacheTrackAbcFile.h
@@ -35,6 +35,10 @@ public:
 
 	bool GetMeshData(int32 SampleIndex, FGeometryCacheMeshData& OutMeshData);
 
+private:
+	void Reset();
+	void ShowNotification(const FText& Text);
+
 private:
 	FGeometryCacheMeshData MeshData;
 	FGeometryCacheTrackSampleInfo SampleInfo;
