CIS shadow-variable fix

diff --git a/Engine/Plugins/Experimental/GeometryCacheAbcFile/Source/GeometryCacheAbcFile/Private/GeometryCacheAbcFileSceneProxy.cpp b/Engine/Plugins/Experimental/GeometryCacheAbcFile/Source/GeometryCacheAbcFile/Private/GeometryCacheAbcFileSceneProxy.cpp
--- a/Engine/Plugins/Experimental/GeometryCacheAbcFile/Source/GeometryCacheAbcFile/Private/GeometryCacheAbcFileSceneProxy.cpp
+++ b/Engine/Plugins/Experimental/GeometryCacheAbcFile/Source/GeometryCacheAbcFile/Private/GeometryCacheAbcFileSceneProxy.cpp
@@ -45,7 +45,7 @@ const FVisibilitySample& FGeomCacheTrackAbcFileProxy::GetVisibilitySample(float
 	return FVisibilitySample::VisibleSample;
 }
 
-void FGeomCacheTrackAbcFileProxy::FindSampleIndexesFromTime(float Time, bool bLooping, bool bIsPlayingBackwards, int32 &OutFrameIndex, int32 &OutNextFrameIndex, float &InterpolationFactor)
+void FGeomCacheTrackAbcFileProxy::FindSampleIndexesFromTime(float Time, bool bLooping, bool bIsPlayingBackwards, int32 &OutFrameIndex, int32 &OutNextFrameIndex, float &InInterpolationFactor)
 {
 	if (UGeometryCacheTrackAbcFile* AbcTrack = Cast<UGeometryCacheTrackAbcFile>(Track))
 	{
@@ -53,13 +53,13 @@ void FGeomCacheTrackAbcFileProxy::FindSampleIndexesFromTime(float Time, bool bLo
 		int32 LastFrameIndex = AbcTrack->GetEndFrameIndex();
 		OutFrameIndex = ThisFrameIndex;
 		OutNextFrameIndex = OutFrameIndex + 1;
-		InterpolationFactor = 0.f;
+		InInterpolationFactor = 0.f;
 
 		// If playing backwards the logical order of previous and next is reversed
 		if (bIsPlayingBackwards)
 		{
 			Swap(OutFrameIndex, OutNextFrameIndex);
-			InterpolationFactor = 1.0f - InterpolationFactor;
+			InInterpolationFactor = 1.0f - InInterpolationFactor;
 		}
 	}
 }
