Fixed crash with GeometryCacheAbcFile after leaving PIE and reloading Alembic.

diff --git a/Engine/Plugins/Experimental/GeometryCacheAbcFile/Source/GeometryCacheAbcFile/Private/GeometryCacheAbcFileComponent.cpp b/Engine/Plugins/Experimental/GeometryCacheAbcFile/Source/GeometryCacheAbcFile/Private/GeometryCacheAbcFileComponent.cpp
--- a/Engine/Plugins/Experimental/GeometryCacheAbcFile/Source/GeometryCacheAbcFile/Private/GeometryCacheAbcFileComponent.cpp
+++ b/Engine/Plugins/Experimental/GeometryCacheAbcFile/Source/GeometryCacheAbcFile/Private/GeometryCacheAbcFileComponent.cpp
@@ -7,9 +7,9 @@
 
 #define LOCTEXT_NAMESPACE "GeometryCacheAbcFileComponent"
 
-UGeometryCacheAbcFileComponent::UGeometryCacheAbcFileComponent()
+UGeometryCacheAbcFileComponent::UGeometryCacheAbcFileComponent(const FObjectInitializer& ObjectInitializer)
 {
-	AbcSettings = NewObject<UAbcImportSettings>(this, TEXT("AbcSettings"));
+	AbcSettings = ObjectInitializer.CreateDefaultSubobject<UAbcImportSettings>(this, TEXT("AbcSettings"));
 }
 
 #if WITH_EDITOR
diff --git a/Engine/Plugins/Experimental/GeometryCacheAbcFile/Source/GeometryCacheAbcFile/Public/GeometryCacheAbcFileComponent.h b/Engine/Plugins/Experimental/GeometryCacheAbcFile/Source/GeometryCacheAbcFile/Public/GeometryCacheAbcFileComponent.h
--- a/Engine/Plugins/Experimental/GeometryCacheAbcFile/Source/GeometryCacheAbcFile/Public/GeometryCacheAbcFileComponent.h
+++ b/Engine/Plugins/Experimental/GeometryCacheAbcFile/Source/GeometryCacheAbcFile/Public/GeometryCacheAbcFileComponent.h
@@ -16,7 +16,7 @@ class GEOMETRYCACHEABCFILE_API UGeometryCacheAbcFileComponent : public UGeometry
 {
 	GENERATED_BODY()
 
-	UGeometryCacheAbcFileComponent();
+	UGeometryCacheAbcFileComponent(const FObjectInitializer& ObjectInitializer);
 
 public:
 
