Fixed crash with GeometryCacheAbcFile after leaving PIE and reloading Alembic.

diff --git a/Engine/Plugins/Experimental/GeometryCacheAbcFile/Source/GeometryCacheAbcFile/Public/GeometryCacheAbcFileComponent.h b/Engine/Plugins/Experimental/GeometryCacheAbcFile/Source/GeometryCacheAbcFile/Public/GeometryCacheAbcFileComponent.h
--- a/Engine/Plugins/Experimental/GeometryCacheAbcFile/Source/GeometryCacheAbcFile/Public/GeometryCacheAbcFileComponent.h
+++ b/Engine/Plugins/Experimental/GeometryCacheAbcFile/Source/GeometryCacheAbcFile/Public/GeometryCacheAbcFileComponent.h
@@ -50,5 +50,6 @@ public:
 protected:
 	void InitializeGeometryCache();
 
+	UPROPERTY(Transient)
 	UAbcImportSettings* AbcSettings;
 };
