Streaming GeometryCache:
* Implemented memory limits in streaming GeometryCache from Alembic. There are two settings available under Project Settings > Plugins > Geometry Cache to control the look-ahead buffer and the overall memory limit for all streams.

diff --git a/Engine/Plugins/Experimental/AlembicImporter/Source/AlembicLibrary/Private/AbcFile.cpp b/Engine/Plugins/Experimental/AlembicImporter/Source/AlembicLibrary/Private/AbcFile.cpp
--- a/Engine/Plugins/Experimental/AlembicImporter/Source/AlembicLibrary/Private/AbcFile.cpp
+++ b/Engine/Plugins/Experimental/AlembicImporter/Source/AlembicLibrary/Private/AbcFile.cpp
@@ -120,6 +120,17 @@ EAbcImportError FAbcFile::Open()
 	DateWritten = ANSI_TO_TCHAR(TmpDateWritten.c_str());
 	UserDescription = ANSI_TO_TCHAR(TmpUserDescription.c_str());
 
+	Alembic::Util::Digest ChildrenDigest;
+	if (TopObject.getChildrenHash(ChildrenDigest))
+	{
+		Hash = ANSI_TO_TCHAR(ChildrenDigest.str().c_str());
+	}
+	else
+	{
+		// Fallback without actually hashing the file
+		Hash = FString::Printf(TEXT("%s;%s"), *FPaths::GetBaseFilename(FilePath), *DateWritten);
+	}
+
 	Alembic::Abc::IBox3dProperty ArchiveBoundsProperty = Alembic::AbcGeom::GetIArchiveBounds(Archive, Alembic::Abc::ErrorHandler::kQuietNoopPolicy);
 
 	if (ArchiveBoundsProperty.valid())
@@ -149,7 +160,8 @@ TArray<FAbcFile::FMetaData> FAbcFile::GetArchiveMetaData() const
 		FMetaDataEntry(TEXT("Abc.AppName"), AppName),
 		FMetaDataEntry(TEXT("Abc.LibraryVersion"), LibVersionString),
 		FMetaDataEntry(TEXT("Abc.WrittenOn"), DateWritten),
-		FMetaDataEntry(TEXT("Abc.UserDescription"), UserDescription)
+		FMetaDataEntry(TEXT("Abc.UserDescription"), UserDescription),
+		FMetaDataEntry(TEXT("Abc.Hash"), Hash)
 	});
 
 	for (const FMetaData& CustomAttribute : CustomAttributes)
diff --git a/Engine/Plugins/Experimental/AlembicImporter/Source/AlembicLibrary/Public/AbcFile.h b/Engine/Plugins/Experimental/AlembicImporter/Source/AlembicLibrary/Public/AbcFile.h
--- a/Engine/Plugins/Experimental/AlembicImporter/Source/AlembicLibrary/Public/AbcFile.h
+++ b/Engine/Plugins/Experimental/AlembicImporter/Source/AlembicLibrary/Public/AbcFile.h
@@ -167,6 +167,7 @@ protected:
 	uint32 LibVersion;
 	FString DateWritten;
 	FString UserDescription;
+	FString Hash;
 
 	TMap<FString, FString> CustomAttributes;
 };
diff --git a/Engine/Plugins/Experimental/GeometryCache/Source/GeometryCache/Classes/GeometryCacheMeshData.h b/Engine/Plugins/Experimental/GeometryCache/Source/GeometryCache/Classes/GeometryCacheMeshData.h
--- a/Engine/Plugins/Experimental/GeometryCache/Source/GeometryCache/Classes/GeometryCacheMeshData.h
+++ b/Engine/Plugins/Experimental/GeometryCache/Source/GeometryCache/Classes/GeometryCacheMeshData.h
@@ -136,10 +136,10 @@ struct GEOMETRYCACHE_API FGeometryCacheMeshData
 	FGeometryCacheVertexInfo VertexInfo;
 		
 	/** Serialization for FVertexAnimationSample. */
-	friend FArchive& operator<<(FArchive& Ar, FGeometryCacheMeshData& Mesh);
+	friend GEOMETRYCACHE_API FArchive& operator<<(FArchive& Ar, FGeometryCacheMeshData& Mesh);
 
 	/** Serialization for const FVertexAnimationSample. */
-	friend FArchive& operator<<(FArchive& Ar, const FGeometryCacheMeshData& Mesh)
+	friend GEOMETRYCACHE_API FArchive& operator<<(FArchive& Ar, const FGeometryCacheMeshData& Mesh)
 	{
 		check(Ar.IsSaving());
 		return (Ar << const_cast<FGeometryCacheMeshData&>(Mesh));
diff --git a/Engine/Plugins/Experimental/GeometryCache/Source/GeometryCache/Private/GeometryCacheMeshData.cpp b/Engine/Plugins/Experimental/GeometryCache/Source/GeometryCache/Private/GeometryCacheMeshData.cpp
--- a/Engine/Plugins/Experimental/GeometryCache/Source/GeometryCache/Private/GeometryCacheMeshData.cpp
+++ b/Engine/Plugins/Experimental/GeometryCache/Source/GeometryCache/Private/GeometryCacheMeshData.cpp
@@ -6,7 +6,8 @@
 
 DECLARE_CYCLE_STAT(TEXT("Deserialize Vertices"), STAT_DeserializeVertices, STATGROUP_GeometryCache);
 DECLARE_CYCLE_STAT(TEXT("Deserialize Indices"), STAT_DeserializeIndices, STATGROUP_GeometryCache);
-DECLARE_CYCLE_STAT(TEXT("Deserialize Schnabbels"), STAT_DeserializeSchnabbels, STATGROUP_GeometryCache);
+DECLARE_CYCLE_STAT(TEXT("Deserialize Metadata"), STAT_DeserializeMetadata, STATGROUP_GeometryCache);
+
 
 FArchive& operator<<(FArchive& Ar, FGeometryCacheMeshData& Mesh)
 {
@@ -23,7 +24,7 @@ FArchive& operator<<(FArchive& Ar, FGeometryCacheMeshData& Mesh)
 
 	// Serialize metadata first so we can use it later on
 	{
-		SCOPE_CYCLE_COUNTER(STAT_DeserializeSchnabbels);
+		SCOPE_CYCLE_COUNTER(STAT_DeserializeMetadata);
 		Ar << Mesh.BoundingBox;
 		Ar << Mesh.BatchesInfo;
 		Ar << Mesh.VertexInfo;
@@ -36,30 +37,66 @@ FArchive& operator<<(FArchive& Ar, FGeometryCacheMeshData& Mesh)
 		if (Ar.IsLoading())
 		{
 			Mesh.Positions.SetNumUninitialized(NumVertices);
-			Mesh.TextureCoordinates.SetNumUninitialized(NumVertices);
-			Mesh.TangentsX.SetNumUninitialized(NumVertices);
-			Mesh.TangentsZ.SetNumUninitialized(NumVertices);
-			Mesh.Colors.SetNumUninitialized(NumVertices);
-			Mesh.ImportedVertexNumbers.SetNumUninitialized(NumVertices);
 
-			if (Mesh.VertexInfo.bHasMotionVectors)
+			if (Mesh.VertexInfo.bHasUV0)
 			{
-				Mesh.MotionVectors.SetNumUninitialized(NumVertices);
+				Mesh.TextureCoordinates.SetNumUninitialized(NumVertices);
 			}
-			else
+
+			if (Mesh.VertexInfo.bHasTangentX)
 			{
-				Mesh.MotionVectors.Empty();
+				Mesh.TangentsX.SetNumUninitialized(NumVertices);
+			}
+
+			if (Mesh.VertexInfo.bHasTangentZ)
+			{
+				Mesh.TangentsZ.SetNumUninitialized(NumVertices);
+			}
+
+			if (Mesh.VertexInfo.bHasColor0)
+			{
+				Mesh.Colors.SetNumUninitialized(NumVertices);
+			}
+
+			if (Mesh.VertexInfo.bHasImportedVertexNumbers)
+			{
+				Mesh.ImportedVertexNumbers.SetNumUninitialized(NumVertices);
+			}
+
+			if (Mesh.VertexInfo.bHasMotionVectors)
+			{
+				Mesh.MotionVectors.SetNumUninitialized(NumVertices);
 			}
 		}
 
 		if (Mesh.Positions.Num() > 0)
 		{				
 			Ar.Serialize(&Mesh.Positions[0], Mesh.Positions.Num() * Mesh.Positions.GetTypeSize());
-			Ar.Serialize(&Mesh.TextureCoordinates[0], Mesh.TextureCoordinates.Num() * Mesh.TextureCoordinates.GetTypeSize());
-			Ar.Serialize(&Mesh.TangentsX[0], Mesh.TangentsX.Num() * Mesh.TangentsX.GetTypeSize());
-			Ar.Serialize(&Mesh.TangentsZ[0], Mesh.TangentsZ.Num() * Mesh.TangentsZ.GetTypeSize());
-			Ar.Serialize(&Mesh.Colors[0], Mesh.Colors.Num() * Mesh.Colors.GetTypeSize());/**/
-			Ar.Serialize(&Mesh.ImportedVertexNumbers, Mesh.ImportedVertexNumbers.Num() * Mesh.ImportedVertexNumbers.GetTypeSize());
+
+			if (Mesh.VertexInfo.bHasUV0)
+			{
+				Ar.Serialize(&Mesh.TextureCoordinates[0], Mesh.TextureCoordinates.Num() * Mesh.TextureCoordinates.GetTypeSize());
+			}
+
+			if (Mesh.VertexInfo.bHasTangentX)
+			{
+				Ar.Serialize(&Mesh.TangentsX[0], Mesh.TangentsX.Num() * Mesh.TangentsX.GetTypeSize());
+			}
+
+			if (Mesh.VertexInfo.bHasTangentZ)
+			{
+				Ar.Serialize(&Mesh.TangentsZ[0], Mesh.TangentsZ.Num() * Mesh.TangentsZ.GetTypeSize());
+			}
+
+			if (Mesh.VertexInfo.bHasColor0)
+			{
+				Ar.Serialize(&Mesh.Colors[0], Mesh.Colors.Num() * Mesh.Colors.GetTypeSize());
+			}
+
+			if (Mesh.VertexInfo.bHasImportedVertexNumbers)
+			{
+				Ar.Serialize(&Mesh.ImportedVertexNumbers, Mesh.ImportedVertexNumbers.Num() * Mesh.ImportedVertexNumbers.GetTypeSize());
+			}
 		  
 			if (Mesh.VertexInfo.bHasMotionVectors)
 			{
diff --git a/Engine/Plugins/Experimental/GeometryCache/Source/GeometryCacheStreamer/GeometryCacheStreamer.Build.cs b/Engine/Plugins/Experimental/GeometryCache/Source/GeometryCacheStreamer/GeometryCacheStreamer.Build.cs
--- a/Engine/Plugins/Experimental/GeometryCache/Source/GeometryCacheStreamer/GeometryCacheStreamer.Build.cs
+++ b/Engine/Plugins/Experimental/GeometryCache/Source/GeometryCacheStreamer/GeometryCacheStreamer.Build.cs
@@ -10,6 +10,8 @@ namespace UnrealBuildTool.Rules
 				new string[]
 				{
 					"Core",
+					"CoreUObject",
+					"DeveloperSettings",
 					"GeometryCache",
 					"Slate",
 					"SlateCore"
diff --git a/Engine/Plugins/Experimental/GeometryCache/Source/GeometryCacheStreamer/Private/GeometryCacheStreamer.cpp b/Engine/Plugins/Experimental/GeometryCache/Source/GeometryCacheStreamer/Private/GeometryCacheStreamer.cpp
--- a/Engine/Plugins/Experimental/GeometryCache/Source/GeometryCacheStreamer/Private/GeometryCacheStreamer.cpp
+++ b/Engine/Plugins/Experimental/GeometryCache/Source/GeometryCacheStreamer/Private/GeometryCacheStreamer.cpp
@@ -5,11 +5,23 @@
 #include "Containers/Ticker.h"
 #include "Framework/Notifications/NotificationManager.h"
 #include "GeometryCacheMeshData.h"
+#include "GeometryCacheModule.h"
+#include "GeometryCacheStreamerSettings.h"
 #include "IGeometryCacheStream.h"
+#include "Stats/Stats.h"
 #include "Widgets/Notifications/SNotificationList.h"
 
 #define LOCTEXT_NAMESPACE "GeometryCacheStreamer"
 
+DECLARE_DWORD_COUNTER_STAT(TEXT("GeometryCache Streams: Count"), STAT_GeometryCacheStream_Count, STATGROUP_GeometryCache);
+DECLARE_MEMORY_STAT(TEXT("GeometryCache Streams: Memory Used"), STAT_GeometryCacheStream_MemoryUsed, STATGROUP_GeometryCache);
+
+static bool GShowGeometryCacheStreamerNotification = true;
+static FAutoConsoleVariableRef CVarGeometryCacheStreamerShowNotification(
+	TEXT("GeometryCache.Streamer.ShowNotification"),
+	GShowGeometryCacheStreamerNotification,
+	TEXT("Show notification while the GeometryCache streamer is streaming data"));
+
 class FGeometryCacheStreamer : public IGeometryCacheStreamer
 {
 public:
@@ -27,6 +39,9 @@ public:
 	virtual bool TryGetFrameData(UGeometryCacheTrack* AbcTrack, int32 FrameIndex, FGeometryCacheMeshData& OutMeshData) override;
 	//~ End IGeometryCacheStreamer Interface
 
+private:
+	void BalanceStreams();
+
 private:
 	FTSTicker::FDelegateHandle TickHandle;
 
@@ -67,17 +82,18 @@ FGeometryCacheStreamer::~FGeometryCacheStreamer()
 void FGeometryCacheStreamer::Tick(float Time)
 {
 	// The sole purpose of the Streamer is to schedule the streams for read at every engine loop
+	BalanceStreams();
 
 	// First step is to process the results from the previously scheduled read requests
 	// to figure out the number of reads still in progress
 	int32 NumFramesToStream = 0;
-	for (FTracksToStreams::TIterator Iterator = TracksToStreams.CreateIterator(); Iterator; ++Iterator)
+	for (const auto& Pair : TracksToStreams)
 	{
 		TArray<int32> FramesCompleted;
-		IGeometryCacheStream* Stream = Iterator->Value;
+		IGeometryCacheStream* Stream = Pair.Value;
 		Stream->UpdateRequestStatus(FramesCompleted);
 		NumReads -= FramesCompleted.Num();
-		NumFramesToStream += Stream->GetFramesNeeded().Num();
+		NumFramesToStream += Stream->GetNumFramesNeeded();
 	}
 
 	// Now, schedule new read requests according to the number of concurrent reads available
@@ -106,10 +122,9 @@ void FGeometryCacheStreamer::Tick(float Time)
 			}
 
 			IGeometryCacheStream* Stream = Streams[CurrentIndex];
-			const TArray<int32>& FramesNeeded = Stream->GetFramesNeeded();
-			if (FramesNeeded.Num() > 0)
+			if (Stream->GetNumFramesNeeded() > 0)
 			{
-				if (Stream->RequestFrameData(FramesNeeded[0]))
+				if (Stream->RequestFrameData())
 				{
 					// Stream was able to handle the read request so there's one less available
 					++NumReads;
@@ -129,8 +144,11 @@ void FGeometryCacheStreamer::Tick(float Time)
 		}
 	}
 
-	// Display a streaming progress notification if there are any frames to stream
-	if (!StreamingNotification.IsValid() && NumFramesToStream > 0)
+	INC_DWORD_STAT_BY(STAT_GeometryCacheStream_Count, NumStreams);
+
+	// Display a streaming progress notification if the the number of frames to stream is above the threshold
+	// to prevent spamming notifications when there's only a few frames to stream
+	if (GShowGeometryCacheStreamerNotification && !StreamingNotification.IsValid() && NumFramesToStream >= 24)
 	{
 		FText UpdateText = FText::Format(LOCTEXT("GeoCacheStreamingUpdate", "Streaming GeometryCache: {0} frames remaining"), FText::AsNumber(NumFramesToStream));
 		FNotificationInfo Info(UpdateText);
@@ -164,6 +182,43 @@ void FGeometryCacheStreamer::Tick(float Time)
 	}
 }
 
+void FGeometryCacheStreamer::BalanceStreams()
+{
+	// Collect the stats for all streams
+	FGeometryCacheStreamStats AggregatedStats;
+	for (const auto& Pair : TracksToStreams)
+	{
+		IGeometryCacheStream* Stream = Pair.Value;
+		const FGeometryCacheStreamStats& StreamStats = Stream->GetStreamStats();
+		AggregatedStats.MemoryUsed += StreamStats.MemoryUsed;
+		AggregatedStats.CachedDuration = FMath::Max(AggregatedStats.CachedDuration, StreamStats.CachedDuration);
+		AggregatedStats.AverageBitrate += StreamStats.AverageBitrate;
+	}
+
+	SET_MEMORY_STAT(STAT_GeometryCacheStream_MemoryUsed, AggregatedStats.MemoryUsed * 1024 * 1024);
+
+	const UGeometryCacheStreamerSettings* Settings = GetDefault<UGeometryCacheStreamerSettings>();
+	const float MemoryLimit = Settings->MaxMemoryAllowed;
+	const float DurationLimit = Settings->LookAheadBuffer;
+
+	// Apply memory limits if over budget
+	// Might want to also update the limits if under budget
+	if (AggregatedStats.MemoryUsed > MemoryLimit)
+	{
+		// Allocate per-stream budget based on the duration computed from the total bitrate and given memory limit
+		const float AvgDuration = MemoryLimit / AggregatedStats.AverageBitrate;
+		for (const auto& Pair : TracksToStreams)
+		{
+			IGeometryCacheStream* Stream = Pair.Value;
+			const FGeometryCacheStreamStats& StreamStats = Stream->GetStreamStats();
+
+			const float MaxDuration = FMath::Clamp(AvgDuration, 0.f, DurationLimit);
+			const float MaxMem = AvgDuration * StreamStats.AverageBitrate;
+			Stream->SetLimits(MaxMem, MaxDuration);
+		}
+	}
+}
+
 void FGeometryCacheStreamer::RegisterTrack(UGeometryCacheTrack* AbcTrack, IGeometryCacheStream* Stream)
 {
 	check(AbcTrack && Stream && !TracksToStreams.Contains(AbcTrack));
diff --git a/Engine/Plugins/Experimental/GeometryCache/Source/GeometryCacheStreamer/Private/GeometryCacheStreamerSettings.cpp b/Engine/Plugins/Experimental/GeometryCache/Source/GeometryCacheStreamer/Private/GeometryCacheStreamerSettings.cpp
new file mode 100644
--- /dev/null
+++ b/Engine/Plugins/Experimental/GeometryCache/Source/GeometryCacheStreamer/Private/GeometryCacheStreamerSettings.cpp
@@ -0,0 +1,12 @@
+// Copyright Epic Games, Inc. All Rights Reserved.
+
+#include "GeometryCacheStreamerSettings.h"
+	 
+UGeometryCacheStreamerSettings::UGeometryCacheStreamerSettings()
+{
+	CategoryName = TEXT("Plugins");
+	SectionName = TEXT("Geometry Cache");
+
+	LookAheadBuffer = 4.0f;
+	MaxMemoryAllowed = 4096.0f;
+}
diff --git a/Engine/Plugins/Experimental/GeometryCache/Source/GeometryCacheStreamer/Public/GeometryCacheStreamerSettings.h b/Engine/Plugins/Experimental/GeometryCache/Source/GeometryCacheStreamer/Public/GeometryCacheStreamerSettings.h
new file mode 100644
--- /dev/null
+++ b/Engine/Plugins/Experimental/GeometryCache/Source/GeometryCacheStreamer/Public/GeometryCacheStreamerSettings.h
@@ -0,0 +1,27 @@
+// Copyright Epic Games, Inc. All Rights Reserved.
+
+#pragma once
+
+#include "Engine/DeveloperSettings.h"
+#include "UObject/ObjectMacros.h"
+
+#include "GeometryCacheStreamerSettings.generated.h"
+
+/** Settings for the GeometryCache streamer */
+UCLASS(config = Engine, meta = (DisplayName = "Geometry Cache"))
+class GEOMETRYCACHESTREAMER_API UGeometryCacheStreamerSettings : public UDeveloperSettings
+{
+	GENERATED_BODY()
+
+public:
+	 
+	UGeometryCacheStreamerSettings();
+
+	/** The amount of animation (in seconds) to stream ahead of time (per stream) */
+	UPROPERTY(config, EditAnywhere, Category = "Geometry Cache Streamer", meta = (DisplayName = "Look-Ahead Buffer (in seconds)", ClampMin = "0.01", ClampMax = "3600.0"))
+	float LookAheadBuffer;
+
+	/** The maximum total amount of streamed data allowed in memory (for all streams) */
+	UPROPERTY(config, EditAnywhere, Category = "Geometry Cache Streamer", meta = (DisplayName = "Maximum Memory Allowed (in MB)", ClampMin = "1.0", ClampMax = "262144.0"))
+	float MaxMemoryAllowed;
+};
diff --git a/Engine/Plugins/Experimental/GeometryCache/Source/GeometryCacheStreamer/Public/IGeometryCacheStream.h b/Engine/Plugins/Experimental/GeometryCache/Source/GeometryCacheStreamer/Public/IGeometryCacheStream.h
--- a/Engine/Plugins/Experimental/GeometryCache/Source/GeometryCacheStreamer/Public/IGeometryCacheStream.h
+++ b/Engine/Plugins/Experimental/GeometryCache/Source/GeometryCacheStreamer/Public/IGeometryCacheStream.h
@@ -5,6 +5,7 @@
 #include "CoreMinimal.h"
 
 struct FGeometryCacheMeshData;
+struct FGeometryCacheStreamStats;
 
 /** Interface to stream GeometryCacheMeshData asynchronously from any source through the GeometryCacheStreamer */
 class IGeometryCacheStream
@@ -15,11 +16,11 @@ public:
 	/** Prefetch NumFrames starting from the given StartFrameIndex. If no NumFrames given, prefetch the whole stream */
 	virtual void Prefetch(int32 StartFrameIndex, int32 NumFrames = 0) = 0;
 
-	/** Return list of frame indices needed to be loaded */
-	virtual const TArray<int32>& GetFramesNeeded() = 0;
+	/** Return the number of frame indices needed to be loaded */
+	virtual uint32 GetNumFramesNeeded() = 0;
 	
-	/** Request a read of the given FrameIndex. Return true if the request could be handled */
-	virtual bool RequestFrameData(int32 FrameIndex) = 0;
+	/** Request a read of the next frame as determined by the stream. Return true if the request could be handled */
+	virtual bool RequestFrameData() = 0;
 
 	/** Update the status of the read requests currently in progress. Return the frame indices that were completed */
 	virtual void UpdateRequestStatus(TArray<int32>& OutFramesCompleted) = 0;
@@ -31,4 +32,18 @@ public:
 
 	/** Cancel the scheduled read requests. Return the number of requests that were canceled */
 	virtual int32 CancelRequests() = 0;
+
+	/** Return the memory usage and related stats for the stream */
+	virtual const FGeometryCacheStreamStats& GetStreamStats() const = 0;
+
+	/** Set the memory usage limits for the stream */
+	virtual void SetLimits(float MaxMemoryAllowed, float MaxCachedDuration) = 0;
+};
+
+struct FGeometryCacheStreamStats
+{
+	uint32 NumCachedFrames = 0; // number of frames currently resident in memory
+	float CachedDuration = 0.0f;// in seconds
+	float MemoryUsed = 0.0f;	// in MB
+	float AverageBitrate = 0.0f;// in MB/s
 };
diff --git a/Engine/Plugins/Experimental/GeometryCacheAbcFile/Source/GeometryCacheAbcFile/GeometryCacheAbcFile.Build.cs b/Engine/Plugins/Experimental/GeometryCacheAbcFile/Source/GeometryCacheAbcFile/GeometryCacheAbcFile.Build.cs
--- a/Engine/Plugins/Experimental/GeometryCacheAbcFile/Source/GeometryCacheAbcFile/GeometryCacheAbcFile.Build.cs
+++ b/Engine/Plugins/Experimental/GeometryCacheAbcFile/Source/GeometryCacheAbcFile/GeometryCacheAbcFile.Build.cs
@@ -28,5 +28,7 @@ public class GeometryCacheAbcFile : ModuleRules
                 "UnrealEd"
             }
         );
+
+		PrivateIncludePathModuleNames.Add("DerivedDataCache");
 	}
 }
diff --git a/Engine/Plugins/Experimental/GeometryCacheAbcFile/Source/GeometryCacheAbcFile/Private/GeometryCacheAbcFileComponent.cpp b/Engine/Plugins/Experimental/GeometryCacheAbcFile/Source/GeometryCacheAbcFile/Private/GeometryCacheAbcFileComponent.cpp
--- a/Engine/Plugins/Experimental/GeometryCacheAbcFile/Source/GeometryCacheAbcFile/Private/GeometryCacheAbcFileComponent.cpp
+++ b/Engine/Plugins/Experimental/GeometryCacheAbcFile/Source/GeometryCacheAbcFile/Private/GeometryCacheAbcFileComponent.cpp
@@ -10,6 +10,7 @@
 UGeometryCacheAbcFileComponent::UGeometryCacheAbcFileComponent(const FObjectInitializer& ObjectInitializer)
 {
 	AbcSettings = ObjectInitializer.CreateDefaultSubobject<UAbcImportSettings>(this, TEXT("AbcSettings"));
+	SamplingSettings.bSkipEmpty = true;
 	NormalGenerationSettings.bForceOneSmoothingGroupPerObject = true;
 }
 
@@ -58,6 +59,7 @@ void UGeometryCacheAbcFileComponent::ReloadAbcFile()
 	AbcSettings->MaterialSettings = MaterialSettings;
 	AbcSettings->ConversionSettings = ConversionSettings;
 	AbcSettings->NormalGenerationSettings = NormalGenerationSettings;
+	AbcSettings->GeometryCacheSettings = GeometryCacheSettings;
 
 	FString FilePath(AlembicFilePath.FilePath);
 	if (FPaths::IsRelative(FilePath))
diff --git a/Engine/Plugins/Experimental/GeometryCacheAbcFile/Source/GeometryCacheAbcFile/Private/GeometryCacheAbcFileComponentCustomization.cpp b/Engine/Plugins/Experimental/GeometryCacheAbcFile/Source/GeometryCacheAbcFile/Private/GeometryCacheAbcFileComponentCustomization.cpp
--- a/Engine/Plugins/Experimental/GeometryCacheAbcFile/Source/GeometryCacheAbcFile/Private/GeometryCacheAbcFileComponentCustomization.cpp
+++ b/Engine/Plugins/Experimental/GeometryCacheAbcFile/Source/GeometryCacheAbcFile/Private/GeometryCacheAbcFileComponentCustomization.cpp
@@ -36,6 +36,12 @@ void FGeometryCacheAbcFileComponentCustomization::CustomizeDetails(IDetailLayout
 	DetailBuilder.HideProperty("SamplingSettings.FrameSteps");
 	DetailBuilder.HideProperty("SamplingSettings.TimeSteps");
 
+	// Streaming Alembic doesn't do optimization nor compression
+	DetailBuilder.HideProperty("GeometryCacheSettings.bApplyConstantTopologyOptimizations");
+	DetailBuilder.HideProperty("GeometryCacheSettings.bOptimizeIndexBuffers");
+	DetailBuilder.HideProperty("GeometryCacheSettings.CompressedPositionPrecision");
+	DetailBuilder.HideProperty("GeometryCacheSettings.CompressedTextureCoordinatesNumberOfBits");
+
 	// Add a button to reload the Alembic file and apply the settings
 	IDetailCategoryBuilder& Category = DetailBuilder.EditCategory("Alembic");
 
diff --git a/Engine/Plugins/Experimental/GeometryCacheAbcFile/Source/GeometryCacheAbcFile/Private/GeometryCacheAbcStream.cpp b/Engine/Plugins/Experimental/GeometryCacheAbcFile/Source/GeometryCacheAbcFile/Private/GeometryCacheAbcStream.cpp
--- a/Engine/Plugins/Experimental/GeometryCacheAbcFile/Source/GeometryCacheAbcFile/Private/GeometryCacheAbcStream.cpp
+++ b/Engine/Plugins/Experimental/GeometryCacheAbcFile/Source/GeometryCacheAbcFile/Private/GeometryCacheAbcStream.cpp
@@ -4,25 +4,83 @@
 #include "AbcFile.h"
 #include "AbcUtilities.h"
 #include "Async/Async.h"
+#include "DerivedDataCacheInterface.h"
 #include "GeometryCacheMeshData.h"
+#include "GeometryCacheStreamerSettings.h"
 #include "GeometryCacheTrackAbcFile.h"
+#include "Math/NumericLimits.h"
+#include "Misc/CoreMisc.h"
 #include "Misc/ScopeLock.h"
 
+static bool GAbcStreamCacheInDDC = true;
+static FAutoConsoleVariableRef CVarAbcStreamCacheInDDC(
+	TEXT("GeometryCache.Streamer.AbcStream.CacheInDDC"),
+	GAbcStreamCacheInDDC,
+	TEXT("Cache the streamed Alembic mesh data in the DDC"));
+
 // Max read concurrency is 8 due to limitation in AbcFile
 static int32 kAbcReadConcurrency = 8;
 
+enum class EAbcStreamReadRequestStatus
+{
+	Scheduled,
+	Completed,
+	Cancelled
+};
+
 struct FGeometryCacheAbcStreamReadRequest
 {
-	FGeometryCacheMeshData* MeshData;
-	int32 ReadIndex;
-	int32 FrameIndex;
-	bool Completed;
+	FGeometryCacheMeshData* MeshData = nullptr;
+	int32 ReadIndex = 0;
+	int32 FrameIndex = 0;
+	EAbcStreamReadRequestStatus Status = EAbcStreamReadRequestStatus::Scheduled;
+};
+
+// If AbcStream derived data needs to be rebuilt (new format, serialization
+// differences, etc.) replace the version GUID below with a new one.
+// In case of merge conflicts with DDC versions, you MUST generate a new GUID
+// and set this new GUID as the version.
+#define ABCSTREAM_DERIVED_DATA_VERSION TEXT("88025D2E38A54CF29FA5A6CAE686B013")
+
+class FAbcStreamDDCUtils
+{
+private:
+	static const FString& GetAbcStreamDerivedDataVersion()
+	{
+		static FString CachedVersionString(ABCSTREAM_DERIVED_DATA_VERSION);
+		return CachedVersionString;
+	}
+
+	static FString BuildDerivedDataKey(const FString& KeySuffix)
+	{
+		return FDerivedDataCacheInterface::BuildCacheKey(TEXT("ABCSTREAM_"), *GetAbcStreamDerivedDataVersion(), *KeySuffix);
+	}
+
+	static FString BuildAbcStreamDerivedDataKeySuffix(UGeometryCacheTrackAbcFile* AbcTrack, int32 FrameIndex)
+	{
+		return AbcTrack->GetAbcTrackHash() + TEXT("_") + LexToString(FrameIndex);
+	}
+
+public:
+	static FString GetAbcStreamDDCKey(UGeometryCacheTrackAbcFile* AbcTrack, int32 FrameIndex)
+	{
+		return BuildDerivedDataKey(BuildAbcStreamDerivedDataKeySuffix(AbcTrack, FrameIndex));
+	}
 };
 
 FGeometryCacheAbcStream::FGeometryCacheAbcStream(UGeometryCacheTrackAbcFile* InAbcTrack)
 : AbcTrack(InAbcTrack)
 , bCancellationRequested(false)
+, LastAccessedFrameIndex(0)
+, MaxCachedFrames(0)
+, MaxCachedDuration(0.f)
+, MaxMemAllowed(TNumericLimits<float>::Max())
+, MemoryUsed(0.f)
+, bCacheNeedsUpdate(false)
 {
+	SecondsPerFrame = AbcTrack->GetAbcFile().GetSecondsPerFrame();
+	Hash = AbcTrack->GetAbcTrackHash();
+
 	for (int32 Index = 0; Index < kAbcReadConcurrency; ++Index)
 	{
 		// Populate the ReadIndices. Note that it used as a stack
@@ -31,7 +89,6 @@ FGeometryCacheAbcStream::FGeometryCacheAbcStream(UGeometryCacheTrackAbcFile* InA
 		// Populate pool of reusable ReadRequests
 		ReadRequestsPool.Add(new FGeometryCacheAbcStreamReadRequest());
 	}
-
 }
 
 FGeometryCacheAbcStream::~FGeometryCacheAbcStream()
@@ -39,10 +96,12 @@ FGeometryCacheAbcStream::~FGeometryCacheAbcStream()
 	CancelRequests();
 
 	// Delete all the cached MeshData
-	FScopeLock Lock(&CriticalSection);
-	for (FFrameIndexToMeshData::TIterator Iterator = FramesAvailable.CreateIterator(); Iterator; ++Iterator)
+	FReadScopeLock ReadLock(FramesAvailableLock);
+	for (const auto& Pair : FramesAvailable)
 	{
-		delete Iterator->Value;
+		FGeometryCacheMeshData* MeshData = Pair.Value;
+		DecrementMemoryStat(*MeshData);
+		delete MeshData;
 	}
 
 	// And ReadRequests from the pool
@@ -54,7 +113,7 @@ FGeometryCacheAbcStream::~FGeometryCacheAbcStream()
 
 int32 FGeometryCacheAbcStream::CancelRequests()
 {
-	TGuardValue<TAtomic<bool>, bool> CancellationRequested(bCancellationRequested, true);
+	TGuardValue<std::atomic<bool>, bool> CancellationRequested(bCancellationRequested, true);
 
 	// Clear the FramesNeeded to prevent scheduling further reads
 	FramesNeeded.Empty();
@@ -73,30 +132,87 @@ int32 FGeometryCacheAbcStream::CancelRequests()
 	return CompletedFrames.Num();
 }
 
-bool FGeometryCacheAbcStream::RequestFrameData(int32 FrameIndex)
+void FGeometryCacheAbcStream::GetMeshData(int32 FrameIndex, int32 ConcurrencyIndex, FGeometryCacheMeshData& MeshData)
+{
+	// Get the mesh data straight from the Alembic file or from the DDC if it's already cached
+	if (GAbcStreamCacheInDDC)
+	{
+		const FString DerivedDataKey = FAbcStreamDDCUtils::GetAbcStreamDDCKey(AbcTrack, FrameIndex);
+		const FString& AbcFile = AbcTrack->GetSourceFile();
+
+		TArray<uint8> DerivedData;
+		if (GetDerivedDataCacheRef().GetSynchronous(*DerivedDataKey, DerivedData, AbcFile))
+		{
+			FMemoryReader Ar(DerivedData, /*bIsPersistent=*/ true);
+			Ar << MeshData;
+		}
+		else
+		{
+			FAbcUtilities::GetFrameMeshData(AbcTrack->GetAbcFile(), FrameIndex, MeshData, ConcurrencyIndex);
+
+			FMemoryWriter Ar(DerivedData, true);
+			Ar << MeshData;
+
+			GetDerivedDataCacheRef().Put(*DerivedDataKey, DerivedData, AbcFile);
+		}
+	}
+	else
+	{
+		// Synchronously load the requested frame data
+		FAbcUtilities::GetFrameMeshData(AbcTrack->GetAbcFile(), FrameIndex, MeshData, ConcurrencyIndex);
+	}
+}
+
+bool FGeometryCacheAbcStream::RequestFrameData()
 {
 	check(IsInGameThread());
 
-	// Don't schedule the same FrameIndex twice
-	for (int32 Index = 0; Index < FramesRequested.Num(); ++Index)
+	if (FramesNeeded.Num() == 0)
+	{
+		return false;
+	}
+
+	FReadScopeLock ReadLock(FramesAvailableLock);
+
+	// Get the next frame index to read that is not already available and not in flight
+	int32 FrameIndex = FramesNeeded[0];
+	bool bFrameIndexValid = false;
+	while (FramesNeeded.Num() && !bFrameIndexValid)
 	{
-		if (FramesRequested[Index]->FrameIndex == FrameIndex)
+		FrameIndex = FramesNeeded[0];
+		bFrameIndexValid = true;
+		if (FramesAvailable.Contains(FrameIndex))
 		{
-			return true;
+			FramesNeeded.Remove(FrameIndex);
+			bFrameIndexValid = false;
+		}
+		for (const FGeometryCacheAbcStreamReadRequest* Request : FramesRequested)
+		{
+			if (Request->FrameIndex == FrameIndex)
+			{
+				FramesNeeded.Remove(FrameIndex);
+				bFrameIndexValid = false;
+				break;
+			}
 		}
 	}
 
+	if (!bFrameIndexValid)
+	{
+		return false;
+	}
+
 	if (ReadIndices.Num() > 0)
 	{
 		// Get any ReadIndex available
-		const int32 ReadIndex = ReadIndices.Pop();
+		const int32 ReadIndex = ReadIndices.Pop(false);
 
 		// Take the ReadRequest from the pool at ReadIndex and initialize it
 		FGeometryCacheAbcStreamReadRequest*& ReadRequest = ReadRequestsPool[ReadIndex];
 		ReadRequest->FrameIndex = FrameIndex;
 		ReadRequest->ReadIndex = ReadIndex;
 		ReadRequest->MeshData = new FGeometryCacheMeshData;
-		ReadRequest->Completed = false;
+		ReadRequest->Status = EAbcStreamReadRequestStatus::Scheduled;
 
 		// Change the frame status from needed to requested
 		FramesNeeded.Remove(FrameIndex);
@@ -108,9 +224,13 @@ bool FGeometryCacheAbcStream::RequestFrameData(int32 FrameIndex)
 			{
 				if (!bCancellationRequested)
 				{
-					FAbcUtilities::GetFrameMeshData(AbcTrack->GetAbcFile(), ReadRequest->FrameIndex, *ReadRequest->MeshData, ReadRequest->ReadIndex);
+					GetMeshData(ReadRequest->FrameIndex, ReadRequest->ReadIndex, *ReadRequest->MeshData);
+					ReadRequest->Status = EAbcStreamReadRequestStatus::Completed;
+				}
+				else
+				{
+					ReadRequest->Status = EAbcStreamReadRequestStatus::Cancelled;
 				}
-				ReadRequest->Completed = true;
 			});
 
 		return true;
@@ -122,26 +242,93 @@ void FGeometryCacheAbcStream::UpdateRequestStatus(TArray<int32>& OutFramesComple
 {
 	check(IsInGameThread());
 
-	FScopeLock Lock(&CriticalSection);
+	// Check if the cache needs to be updated either because the frame has advanced or because there's a new memory limit
+	if (bCacheNeedsUpdate)
+	{
+		bCacheNeedsUpdate = false;
+		UpdateFramesNeeded(LastAccessedFrameIndex, MaxCachedFrames);
+
+		// Remove the unneeded frames, those that are available but don't need to be cached
+		TArray<int32> UnneededFrames;
+		{
+			FReadScopeLock ReadLock(FramesAvailableLock);
+			for (const auto& Pair : FramesAvailable)
+			{
+				int32 FrameIndex = Pair.Key;
+				if (!FramesToBeCached.Contains(FrameIndex))
+				{
+					FGeometryCacheMeshData* MeshData = Pair.Value;
+					DecrementMemoryStat(*MeshData);
+					delete MeshData;
+
+					UnneededFrames.Add(FrameIndex);
+				}
+			}
+		}
+
+		{
+			FWriteScopeLock WriteLock(FramesAvailableLock);
+			for (int32 FrameIndex : UnneededFrames)
+			{
+				FramesAvailable.Remove(FrameIndex);
+			}
+		}
+	}
+
+	FWriteScopeLock WriteLock(FramesAvailableLock);
 
 	// Check the completion status of the read requests in progress
 	TArray<FGeometryCacheAbcStreamReadRequest*> CompletedRequests;
 	for (FGeometryCacheAbcStreamReadRequest* ReadRequest : FramesRequested)
 	{
-		if (ReadRequest->Completed)
+		// A cancelled read is still considered completed, it just hasn't read any data
+		if (ReadRequest->Status == EAbcStreamReadRequestStatus::Completed ||
+			ReadRequest->Status == EAbcStreamReadRequestStatus::Cancelled)
 		{
 			// Queue for removal after iterating
 			CompletedRequests.Add(ReadRequest);
 
-			// Cache result of read for retrieval later
-			// #ueent_todo: Implement some sort of LRU cache to reduce memory usage
-			FramesAvailable.Add(ReadRequest->FrameIndex, ReadRequest->MeshData);
+			// A cancelled read has an allocated mesh data that needs to be deleted
+			bool bDeleteMeshData = ReadRequest->Status == EAbcStreamReadRequestStatus::Cancelled;
+			if (ReadRequest->Status == EAbcStreamReadRequestStatus::Completed)
+			{
+				FResourceSizeEx ResSize;
+				ReadRequest->MeshData->GetResourceSizeEx(ResSize);
+				float FrameDataSize = float(ResSize.GetTotalMemoryBytes()) / (1024 * 1024);
+
+				if (MemoryUsed + FrameDataSize < MaxMemAllowed)
+				{
+					if (!FramesAvailable.Contains(ReadRequest->FrameIndex))
+					{
+						// Cache result of read for retrieval later
+						FramesAvailable.Add(ReadRequest->FrameIndex, ReadRequest->MeshData);
+						IncrementMemoryStat(*ReadRequest->MeshData);
+					}
+					else
+					{
+						// The requested frame was already available, just delete it
+						bDeleteMeshData = true;
+					}
+				}
+				else
+				{
+					// The frame doesn't fit in the allowed memory budget so delete it
+					// It should be requested again later if it's still needed
+					bDeleteMeshData = true;
+				}
+			}
 
 			// Push back the ReadIndex for reuse
 			ReadIndices.Push(ReadRequest->ReadIndex);
 
 			// Output the completed frame
 			OutFramesCompleted.Add(ReadRequest->FrameIndex);
+
+			if (bDeleteMeshData)
+			{
+				delete ReadRequest->MeshData;
+				ReadRequest->MeshData = nullptr;
+			}
 		}
 	}
 
@@ -151,13 +338,62 @@ void FGeometryCacheAbcStream::UpdateRequestStatus(TArray<int32>& OutFramesComple
 	}
 }
 
-void FGeometryCacheAbcStream::Prefetch(int32 StartFrameIndex, int32 NumFrames)
+void FGeometryCacheAbcStream::UpdateFramesNeeded(int32 StartFrameIndex, int32 NumFrames)
 {
+	FramesToBeCached.Empty(NumFrames);
+	FramesNeeded.Empty(NumFrames);
+
 	const FAbcFile& AbcFile = AbcTrack->GetAbcFile();
-	const int32 MaxNumFrames = AbcFile.GetImportNumFrames();
 	const int32 StartIndex = AbcFile.GetStartFrameIndex();
 	const int32 EndIndex = AbcFile.GetEndFrameIndex();
 
+	// FramesToBeCached are the frame indices that are required for playback, available or not,
+	// while FramesNeeded are the frame indices that are not available yet so they need to be read
+	auto AddFrameIndex = [this, &NumFrames](int32 FrameIndex)
+	{
+		FramesToBeCached.Add(FrameIndex);
+		if (!FramesAvailable.Contains(FrameIndex))
+		{
+			FramesNeeded.Add(FrameIndex);
+		}
+		--NumFrames;
+	};
+
+	FReadScopeLock ReadLock(FramesAvailableLock);
+
+	StartFrameIndex = FMath::Clamp(StartFrameIndex, StartIndex, EndIndex);
+
+	// Also reserve space for the frame before start since playback is double-buffered
+	int PreviousFrameIndex = FMath::Clamp(StartFrameIndex - 1, StartIndex, EndIndex);
+	if (PreviousFrameIndex != StartFrameIndex)
+	{
+		--NumFrames;
+	}
+
+	// Populate the list of frame indices from given StartFrameIndex up to NumFrames or EndIndex
+	for (int32 Index = StartFrameIndex; NumFrames > 0 && Index < EndIndex; ++Index)
+	{
+		AddFrameIndex(Index);
+	}
+
+	// End of the range might have been reached before the requested NumFrames so add the remaining frames starting from StartIndex
+	for (int32 Index = StartIndex; NumFrames > 0 && Index < PreviousFrameIndex; ++Index)
+	{
+		AddFrameIndex(Index);
+	}
+
+	// Frame before start is added at the end to preserve the priority of the other frames
+	if (PreviousFrameIndex != StartFrameIndex)
+	{
+		AddFrameIndex(PreviousFrameIndex);
+	}
+}
+
+void FGeometryCacheAbcStream::Prefetch(int32 StartFrameIndex, int32 NumFrames)
+{
+	const FAbcFile& AbcFile = AbcTrack->GetAbcFile();
+	const int32 MaxNumFrames = AbcFile.GetImportNumFrames();
+
 	// Validate the number of frames to be loaded
 	if (NumFrames == 0)
 	{
@@ -169,18 +405,9 @@ void FGeometryCacheAbcStream::Prefetch(int32 StartFrameIndex, int32 NumFrames)
 		NumFrames = FMath::Clamp(NumFrames, 1, MaxNumFrames);
 	}
 
-	// Populate the list of frames needed to be loaded starting from given StartFrameIndex up to NumFrames or EndIndex
-	StartFrameIndex = FMath::Clamp(StartFrameIndex, StartIndex, EndIndex);
-	for (int32 Index = StartFrameIndex; NumFrames > 0 && Index < EndIndex; ++Index, --NumFrames)
-	{
-		FramesNeeded.Add(Index);
-	}
+	MaxCachedFrames = NumFrames;
 
-	// End of the range might have been reached before the requested NumFrames so add the remaining frames starting from StartIndex
-	for (int32 Index = StartIndex; NumFrames > 0; ++Index, --NumFrames)
-	{
-		FramesNeeded.Add(Index);
-	}
+	UpdateFramesNeeded(StartFrameIndex, NumFrames);
 
 	if (FramesNeeded.Num() > 0)
 	{
@@ -190,19 +417,29 @@ void FGeometryCacheAbcStream::Prefetch(int32 StartFrameIndex, int32 NumFrames)
 	}
 }
 
-const TArray<int32>& FGeometryCacheAbcStream::GetFramesNeeded()
+uint32 FGeometryCacheAbcStream::GetNumFramesNeeded()
 {
-	return FramesNeeded;
+	return FramesNeeded.Num();
 }
 
 bool FGeometryCacheAbcStream::GetFrameData(int32 FrameIndex, FGeometryCacheMeshData& OutMeshData)
 {
-	FScopeLock Lock(&CriticalSection);
-	if (FGeometryCacheMeshData** MeshDataPtr = FramesAvailable.Find(FrameIndex))
+	// This function can be called from the render thread
+	int32 PrevLastAccessedFrameIndex = LastAccessedFrameIndex.exchange(FrameIndex);
+	if (PrevLastAccessedFrameIndex != FrameIndex)
 	{
-		OutMeshData = **MeshDataPtr;
-		return true;
+		bCacheNeedsUpdate = true;
+	}
+
+	{
+		FReadScopeLock ReadLock(FramesAvailableLock);
+		if (FGeometryCacheMeshData** MeshDataPtr = FramesAvailable.Find(FrameIndex))
+		{
+			OutMeshData = **MeshDataPtr;
+			return true;
+		}
 	}
+
 	return false;
 }
 
@@ -210,14 +447,62 @@ void FGeometryCacheAbcStream::LoadFrameData(int32 FrameIndex)
 {
 	check(IsInGameThread());
 
+	FWriteScopeLock WriteLock(FramesAvailableLock);
 	if (FramesAvailable.Contains(FrameIndex))
 	{
 		return;
 	}
 
-	// Synchronously load the requested frame data
 	FGeometryCacheMeshData* MeshData = new FGeometryCacheMeshData;
-	FAbcUtilities::GetFrameMeshData(AbcTrack->GetAbcFile(), FrameIndex, *MeshData, 0);
-
+	GetMeshData(FrameIndex, 0, *MeshData);
 	FramesAvailable.Add(FrameIndex, MeshData);
+	IncrementMemoryStat(*MeshData);
+}
+
+const FGeometryCacheStreamStats& FGeometryCacheAbcStream::GetStreamStats() const
+{
+	check(IsInGameThread());
+
+	const int32 NumFrames = FramesAvailable.Num();
+	const float Secs = SecondsPerFrame * NumFrames;
+
+	Stats.NumCachedFrames = NumFrames;
+	Stats.CachedDuration = Secs;
+	Stats.MemoryUsed = MemoryUsed;
+	Stats.AverageBitrate = MemoryUsed / Secs;
+	return Stats;
+}
+
+void FGeometryCacheAbcStream::SetLimits(float InMaxMemoryAllowed, float InMaxCachedDuration)
+{
+	check(IsInGameThread());
+
+	if (InMaxMemoryAllowed != MaxMemAllowed)
+	{
+		if (InMaxMemoryAllowed < MaxMemAllowed)
+		{
+			bCacheNeedsUpdate = true;
+		}
+		MaxMemAllowed = InMaxMemoryAllowed;
+		MaxCachedDuration = FMath::Min(InMaxCachedDuration, AbcTrack->GetAbcFile().GetImportLength());
+		MaxCachedFrames = FMath::Min(FMath::CeilToInt(MaxCachedDuration / SecondsPerFrame), AbcTrack->GetAbcFile().GetImportNumFrames());
+	}
+}
+
+void FGeometryCacheAbcStream::IncrementMemoryStat(const FGeometryCacheMeshData& MeshData)
+{
+	FResourceSizeEx ResSize;
+	MeshData.GetResourceSizeEx(ResSize);
+
+	float SizeInBytes = ResSize.GetTotalMemoryBytes();
+	MemoryUsed += SizeInBytes / (1024 * 1024);
+}
+
+void FGeometryCacheAbcStream::DecrementMemoryStat(const FGeometryCacheMeshData& MeshData)
+{
+	FResourceSizeEx ResSize;
+	MeshData.GetResourceSizeEx(ResSize);
+
+	float SizeInBytes = ResSize.GetTotalMemoryBytes();
+	MemoryUsed -= SizeInBytes / (1024 * 1024);
 }
diff --git a/Engine/Plugins/Experimental/GeometryCacheAbcFile/Source/GeometryCacheAbcFile/Private/GeometryCacheAbcStream.h b/Engine/Plugins/Experimental/GeometryCacheAbcFile/Source/GeometryCacheAbcFile/Private/GeometryCacheAbcStream.h
--- a/Engine/Plugins/Experimental/GeometryCacheAbcFile/Source/GeometryCacheAbcFile/Private/GeometryCacheAbcStream.h
+++ b/Engine/Plugins/Experimental/GeometryCacheAbcFile/Source/GeometryCacheAbcFile/Private/GeometryCacheAbcStream.h
@@ -3,6 +3,7 @@
 #pragma once
 
 #include "IGeometryCacheStream.h"
+#include <atomic>
 
 struct FGeometryCacheAbcStreamReadRequest;
 class UGeometryCacheTrackAbcFile;
@@ -15,16 +16,22 @@ public:
 
 	//~ Begin IGeometryCacheStream Interface
 	virtual void Prefetch(int32 StartFrameIndex, int32 NumFrames = 0) override;
-	virtual const TArray<int32>& GetFramesNeeded() override;
-	virtual bool RequestFrameData(int32 FrameIndex) override;
+	virtual uint32 GetNumFramesNeeded() override;
+	virtual bool RequestFrameData() override;
 	virtual void UpdateRequestStatus(TArray<int32>& OutFramesCompleted) override;
 	virtual bool GetFrameData(int32 FrameIndex, FGeometryCacheMeshData& OutMeshData) override;
 	virtual int32 CancelRequests() override;
+	virtual const FGeometryCacheStreamStats& GetStreamStats() const override;
+	virtual void SetLimits(float MaxMemoryAllowed, float MaxCachedDuration) override;
 	//~ End IGeometryCacheStream Interface
 
 private:
 
+	void GetMeshData(int32 FrameIndex, int32 ConcurrencyIndex, FGeometryCacheMeshData& OutMeshData);
 	void LoadFrameData(int32 FrameIndex);
+	void UpdateFramesNeeded(int32 StartIndex, int32 NumFrames);
+	void IncrementMemoryStat(const FGeometryCacheMeshData& MeshData);
+	void DecrementMemoryStat(const FGeometryCacheMeshData& MeshData);
 
 	UGeometryCacheTrackAbcFile* AbcTrack;
 
@@ -32,12 +39,22 @@ private:
 	TArray<FGeometryCacheAbcStreamReadRequest*> ReadRequestsPool;
 
 	TArray<int32> FramesNeeded;
+	TArray<int32> FramesToBeCached;
 	TArray<FGeometryCacheAbcStreamReadRequest*> FramesRequested;
 
 	typedef TMap<int32, FGeometryCacheMeshData*> FFrameIndexToMeshData;
 	FFrameIndexToMeshData FramesAvailable;
-
-	FCriticalSection CriticalSection;
-
-	TAtomic<bool> bCancellationRequested;
+	FRWLock FramesAvailableLock;
+
+	std::atomic<bool> bCancellationRequested;
+
+	mutable FGeometryCacheStreamStats Stats;
+	FString Hash;
+	float SecondsPerFrame;
+	std::atomic<int32> LastAccessedFrameIndex;
+	int32 MaxCachedFrames;
+	float MaxCachedDuration;
+	float MaxMemAllowed;
+	float MemoryUsed;
+	std::atomic<bool> bCacheNeedsUpdate;
 };
\ No newline at end of file
diff --git a/Engine/Plugins/Experimental/GeometryCacheAbcFile/Source/GeometryCacheAbcFile/Private/GeometryCacheTrackAbcFile.cpp b/Engine/Plugins/Experimental/GeometryCacheAbcFile/Source/GeometryCacheAbcFile/Private/GeometryCacheTrackAbcFile.cpp
--- a/Engine/Plugins/Experimental/GeometryCacheAbcFile/Source/GeometryCacheAbcFile/Private/GeometryCacheTrackAbcFile.cpp
+++ b/Engine/Plugins/Experimental/GeometryCacheAbcFile/Source/GeometryCacheAbcFile/Private/GeometryCacheTrackAbcFile.cpp
@@ -10,9 +10,11 @@
 #include "GeometryCache.h"
 #include "GeometryCacheAbcStream.h"
 #include "GeometryCacheHelpers.h"
+#include "GeometryCacheStreamerSettings.h"
 #include "IGeometryCacheStreamer.h"
 #include "Logging/LogCategory.h"
 #include "Logging/LogMacros.h"
+#include "Misc/ArchiveMD5.h"
 #include "Misc/Paths.h"
 #include "PackageTools.h"
 #include "Widgets/Notifications/SNotificationList.h"
@@ -21,6 +23,44 @@ DEFINE_LOG_CATEGORY_STATIC(LogGeometryCacheAbcFile, Log, All);
 
 #define LOCTEXT_NAMESPACE "GeometryCacheTrackAbcFile"
 
+namespace UE::GeometryCacheTrackAbcFile::Private
+{
+	void SerializeSettingsForDDC(FArchive& Ar, UAbcImportSettings* Settings)
+	{
+		// Include only settings that would affect how the frame data is generated
+		Ar << Settings->ConversionSettings.bFlipU;
+		Ar << Settings->ConversionSettings.bFlipV;
+		Ar << Settings->ConversionSettings.Rotation;
+		Ar << Settings->ConversionSettings.Scale;
+
+		Ar << Settings->GeometryCacheSettings.bFlattenTracks;
+		Ar << Settings->GeometryCacheSettings.bStoreImportedVertexNumbers;
+		Ar << Settings->GeometryCacheSettings.MotionVectors;
+
+		Ar << Settings->NormalGenerationSettings.bForceOneSmoothingGroupPerObject;
+		Ar << Settings->NormalGenerationSettings.HardEdgeAngleThreshold;
+		Ar << Settings->NormalGenerationSettings.bRecomputeNormals;
+		Ar << Settings->NormalGenerationSettings.bIgnoreDegenerateTriangles;
+		Ar << Settings->NormalGenerationSettings.bSkipComputingTangents;
+	}
+
+	FString ComputeSettingsHash(UAbcImportSettings* Settings)
+	{
+		if (!Settings)
+		{
+			return {};
+		}
+
+		FArchiveMD5 ArMD5;
+		SerializeSettingsForDDC(ArMD5, Settings);
+
+		FMD5Hash MD5Hash;
+		ArMD5.GetHash(MD5Hash);
+
+		return BytesToHex(MD5Hash.GetBytes(), MD5Hash.GetSize());
+	}
+}
+
 UGeometryCacheTrackAbcFile::UGeometryCacheTrackAbcFile()
 : EndFrameIndex(0)
 {
@@ -75,6 +115,7 @@ const bool UGeometryCacheTrackAbcFile::UpdateBoundsData(const float Time, const
 void UGeometryCacheTrackAbcFile::Reset()
 {
 	AbcFile.Reset();
+	Hash.Empty();
 
 	EndFrameIndex = 0;
 	Duration = 0.f;
@@ -131,6 +172,20 @@ bool UGeometryCacheTrackAbcFile::SetSourceFile(const FString& FilePath, UAbcImpo
 
 		Result = AbcFile->Import(AbcSettings);
 
+		// The hash is composed of the Alembic file hash and the settings hash used to import it
+		FString AbcHash;
+		for (const FAbcFile::FMetaData& MetaData : AbcFile->GetArchiveMetaData())
+		{
+			if (MetaData.Key == TEXT("Abc.Hash"))
+			{
+				AbcHash = MetaData.Value;
+				break;
+			}
+		}
+
+		FString SettingsHash = UE::GeometryCacheTrackAbcFile::Private::ComputeSettingsHash(AbcSettings);
+		Hash = AbcHash + TEXT("_") + SettingsHash;
+
 		// Set the end frame after import since it might have been modified due to validation at import
 		EndFrameIndex = AbcSettings->SamplingSettings.FrameEnd;
 
@@ -168,7 +223,10 @@ bool UGeometryCacheTrackAbcFile::SetSourceFile(const FString& FilePath, UAbcImpo
 		Streamer.RegisterTrack(this, Stream);
 
 		const int32 InitialFrameIndex = FindSampleIndexFromTime(InitialTime, bIsLooping);
-		Stream->Prefetch(InitialFrameIndex);
+		const float LookAhead = GetDefault<UGeometryCacheStreamerSettings>()->LookAheadBuffer;
+		const int32 NumFrames = FMath::CeilToInt(LookAhead / AbcFile->GetSecondsPerFrame());
+
+		Stream->Prefetch(InitialFrameIndex, NumFrames);
 		GetMeshData(InitialFrameIndex, MeshData);
 
 		if (MeshData.Positions.Num() == 0)
diff --git a/Engine/Plugins/Experimental/GeometryCacheAbcFile/Source/GeometryCacheAbcFile/Public/GeometryCacheAbcFileComponent.h b/Engine/Plugins/Experimental/GeometryCacheAbcFile/Source/GeometryCacheAbcFile/Public/GeometryCacheAbcFileComponent.h
--- a/Engine/Plugins/Experimental/GeometryCacheAbcFile/Source/GeometryCacheAbcFile/Public/GeometryCacheAbcFileComponent.h
+++ b/Engine/Plugins/Experimental/GeometryCacheAbcFile/Source/GeometryCacheAbcFile/Public/GeometryCacheAbcFileComponent.h
@@ -27,14 +27,17 @@ public:
 	FAbcSamplingSettings SamplingSettings;
 
 	UPROPERTY(EditAnywhere, BlueprintReadWrite, Category = "Alembic")
-	FAbcMaterialSettings MaterialSettings;
+	FAbcConversionSettings ConversionSettings;
 
 	UPROPERTY(EditAnywhere, BlueprintReadWrite, Category = "Alembic")
-	FAbcConversionSettings ConversionSettings;
+	FAbcGeometryCacheSettings GeometryCacheSettings;
 
 	UPROPERTY(EditAnywhere, BlueprintReadWrite, Category = "Alembic")
 	FAbcNormalGenerationSettings NormalGenerationSettings;
 
+	UPROPERTY(EditAnywhere, BlueprintReadWrite, Category = "Alembic")
+	FAbcMaterialSettings MaterialSettings;
+
 public:
 #if WITH_EDITOR
 	virtual void PostEditChangeProperty(FPropertyChangedEvent& PropertyChangedEvent) override;
diff --git a/Engine/Plugins/Experimental/GeometryCacheAbcFile/Source/GeometryCacheAbcFile/Public/GeometryCacheTrackAbcFile.h b/Engine/Plugins/Experimental/GeometryCacheAbcFile/Source/GeometryCacheAbcFile/Public/GeometryCacheTrackAbcFile.h
--- a/Engine/Plugins/Experimental/GeometryCacheAbcFile/Source/GeometryCacheAbcFile/Public/GeometryCacheTrackAbcFile.h
+++ b/Engine/Plugins/Experimental/GeometryCacheAbcFile/Source/GeometryCacheAbcFile/Public/GeometryCacheTrackAbcFile.h
@@ -40,6 +40,8 @@ public:
 
 	FAbcFile& GetAbcFile();
 
+	const FString& GetAbcTrackHash() const { return Hash; }
+
 private:
 	void Reset();
 	void ShowNotification(const FText& Text);
@@ -49,6 +51,7 @@ private:
 	FGeometryCacheTrackSampleInfo SampleInfo;
 	TUniquePtr<FAbcFile> AbcFile;
 	FString SourceFile;
+	FString Hash;
 
 	int32 EndFrameIndex;
 };
diff --git a/Engine/Plugins/Importers/USDImporter/Source/GeometryCacheUSD/Private/GeometryCacheUSDStream.cpp b/Engine/Plugins/Importers/USDImporter/Source/GeometryCacheUSD/Private/GeometryCacheUSDStream.cpp
--- a/Engine/Plugins/Importers/USDImporter/Source/GeometryCacheUSD/Private/GeometryCacheUSDStream.cpp
+++ b/Engine/Plugins/Importers/USDImporter/Source/GeometryCacheUSD/Private/GeometryCacheUSDStream.cpp
@@ -69,10 +69,17 @@ int32 FGeometryCacheUsdStream::CancelRequests()
 	return CompletedFrames.Num();
 }
 
-bool FGeometryCacheUsdStream::RequestFrameData(int32 FrameIndex)
+bool FGeometryCacheUsdStream::RequestFrameData()
 {
 	check(IsInGameThread());
 
+	if (FramesNeeded.Num() == 0)
+	{
+		return false;
+	}
+
+	int32 FrameIndex = FramesNeeded[0];
+
 	// Don't schedule the same FrameIndex twice
 	for (int32 Index = 0; Index < FramesRequested.Num(); ++Index)
 	{
@@ -216,9 +223,9 @@ void FGeometryCacheUsdStream::Prefetch(int32 StartFrameIndex, int32 NumFrames)
 	}
 }
 
-const TArray<int32>& FGeometryCacheUsdStream::GetFramesNeeded()
+uint32 FGeometryCacheUsdStream::GetNumFramesNeeded()
 {
-	return FramesNeeded;
+	return FramesNeeded.Num();
 }
 
 bool FGeometryCacheUsdStream::GetFrameData(int32 FrameIndex, FGeometryCacheMeshData& OutMeshData)
@@ -265,3 +272,13 @@ void FGeometryCacheUsdStream::LoadFrameData(int32 FrameIndex)
 
 	FramesAvailable.Add(FrameIndex, MeshData);
 }
+
+const FGeometryCacheStreamStats& FGeometryCacheUsdStream::GetStreamStats() const
+{
+	static FGeometryCacheStreamStats NoStats;
+	return NoStats;
+}
+
+void FGeometryCacheUsdStream::SetLimits(float MaxMemoryAllowed, float MaxCachedDuration)
+{
+}
diff --git a/Engine/Plugins/Importers/USDImporter/Source/GeometryCacheUSD/Private/GeometryCacheUSDStream.h b/Engine/Plugins/Importers/USDImporter/Source/GeometryCacheUSD/Private/GeometryCacheUSDStream.h
--- a/Engine/Plugins/Importers/USDImporter/Source/GeometryCacheUSD/Private/GeometryCacheUSDStream.h
+++ b/Engine/Plugins/Importers/USDImporter/Source/GeometryCacheUSD/Private/GeometryCacheUSDStream.h
@@ -15,11 +15,13 @@ public:
 
 	//~ Begin IGeometryCacheStream Interface
 	virtual void Prefetch(int32 StartFrameIndex, int32 NumFrames = 0) override;
-	virtual const TArray<int32>& GetFramesNeeded() override;
-	virtual bool RequestFrameData(int32 FrameIndex) override;
+	virtual uint32 GetNumFramesNeeded() override;
+	virtual bool RequestFrameData() override;
 	virtual void UpdateRequestStatus(TArray<int32>& OutFramesCompleted) override;
 	virtual bool GetFrameData(int32 FrameIndex, FGeometryCacheMeshData& OutMeshData) override;
 	virtual int32 CancelRequests() override;
+	virtual const FGeometryCacheStreamStats& GetStreamStats() const override;
+	virtual void SetLimits(float MaxMemoryAllowed, float MaxCachedDuration) override;
 	//~ End IGeometryCacheStream Interface
 
 private:
