Fixed Static Analysis warning Dereferencing NULL pointer 'ReadRequest->MeshData'

diff --git a/Engine/Plugins/Experimental/GeometryCacheAbcFile/Source/GeometryCacheAbcFile/Private/GeometryCacheAbcStream.cpp b/Engine/Plugins/Experimental/GeometryCacheAbcFile/Source/GeometryCacheAbcFile/Private/GeometryCacheAbcStream.cpp
--- a/Engine/Plugins/Experimental/GeometryCacheAbcFile/Source/GeometryCacheAbcFile/Private/GeometryCacheAbcStream.cpp
+++ b/Engine/Plugins/Experimental/GeometryCacheAbcFile/Source/GeometryCacheAbcFile/Private/GeometryCacheAbcStream.cpp
@@ -224,6 +224,7 @@ bool FGeometryCacheAbcStream::RequestFrameData()
 			{
 				if (!bCancellationRequested)
 				{
+					checkSlow(ReadRequest->MeshData);
 					GetMeshData(ReadRequest->FrameIndex, ReadRequest->ReadIndex, *ReadRequest->MeshData);
 					ReadRequest->Status = EAbcStreamReadRequestStatus::Completed;
 				}
@@ -292,6 +293,7 @@ void FGeometryCacheAbcStream::UpdateRequestStatus(TArray<int32>& OutFramesComple
 			bool bDeleteMeshData = ReadRequest->Status == EAbcStreamReadRequestStatus::Cancelled;
 			if (ReadRequest->Status == EAbcStreamReadRequestStatus::Completed)
 			{
+				checkSlow(ReadRequest->MeshData);
 				FResourceSizeEx ResSize;
 				ReadRequest->MeshData->GetResourceSizeEx(ResSize);
 				float FrameDataSize = float(ResSize.GetTotalMemoryBytes()) / (1024 * 1024);
