Streaming GeometryCache:
* Implemented GetMeshDataAtTime for AbcFile track.

diff --git a/Engine/Plugins/Experimental/GeometryCacheAbcFile/Source/GeometryCacheAbcFile/Private/GeometryCacheTrackAbcFile.cpp b/Engine/Plugins/Experimental/GeometryCacheAbcFile/Source/GeometryCacheAbcFile/Private/GeometryCacheTrackAbcFile.cpp
--- a/Engine/Plugins/Experimental/GeometryCacheAbcFile/Source/GeometryCacheAbcFile/Private/GeometryCacheTrackAbcFile.cpp
+++ b/Engine/Plugins/Experimental/GeometryCacheAbcFile/Source/GeometryCacheAbcFile/Private/GeometryCacheTrackAbcFile.cpp
@@ -276,6 +276,13 @@ const FGeometryCacheTrackSampleInfo& UGeometryCacheTrackAbcFile::GetSampleInfo(f
 	return SampleInfo;
 }
 
+bool UGeometryCacheTrackAbcFile::GetMeshDataAtTime(float Time, FGeometryCacheMeshData& OutMeshData)
+{
+	const bool bLooping = true;
+	const int32 SampleIndex = FindSampleIndexFromTime(Time, bLooping);
+	return GetMeshData(SampleIndex, OutMeshData);
+}
+
 bool UGeometryCacheTrackAbcFile::GetMeshData(int32 SampleIndex, FGeometryCacheMeshData& OutMeshData)
 {
 	if (AbcFile)
diff --git a/Engine/Plugins/Experimental/GeometryCacheAbcFile/Source/GeometryCacheAbcFile/Public/GeometryCacheTrackAbcFile.h b/Engine/Plugins/Experimental/GeometryCacheAbcFile/Source/GeometryCacheAbcFile/Public/GeometryCacheTrackAbcFile.h
--- a/Engine/Plugins/Experimental/GeometryCacheAbcFile/Source/GeometryCacheAbcFile/Public/GeometryCacheTrackAbcFile.h
+++ b/Engine/Plugins/Experimental/GeometryCacheAbcFile/Source/GeometryCacheAbcFile/Public/GeometryCacheTrackAbcFile.h
@@ -25,6 +25,7 @@ public:
 	virtual const bool UpdateMeshData(const float Time, const bool bLooping, int32& InOutMeshSampleIndex, FGeometryCacheMeshData*& OutMeshData) override;
 	virtual const bool UpdateBoundsData(const float Time, const bool bLooping, const bool bIsPlayingBackward, int32& InOutBoundsSampleIndex, FBox& OutBounds) override;
 	virtual const FGeometryCacheTrackSampleInfo& GetSampleInfo(float Time, const bool bLooping) override;
+	virtual bool GetMeshDataAtTime(float Time, FGeometryCacheMeshData& OutMeshData) override;
 	//~ End UGeometryCacheTrack Interface.
 
 	bool SetSourceFile(const FString& FilePath, class UAbcImportSettings* AbcSettings, float InitialTime = 0.f, bool bIsLooping = true);
