Streaming GeometryCache:
* Implemented memory management for streaming GeometryCache from USD.
  * Added function to hash a GeomMesh (for use with DDC).
* Added a couple of streaming GeometryCache stats to stat GeometryCache.
  * Made better use of memory when streamer is under budget.

diff --git a/Engine/Plugins/Experimental/GeometryCache/Source/GeometryCache/Classes/GeometryCacheComponent.h b/Engine/Plugins/Experimental/GeometryCache/Source/GeometryCache/Classes/GeometryCacheComponent.h
--- a/Engine/Plugins/Experimental/GeometryCache/Source/GeometryCache/Classes/GeometryCacheComponent.h
+++ b/Engine/Plugins/Experimental/GeometryCache/Source/GeometryCache/Classes/GeometryCacheComponent.h
@@ -233,7 +233,7 @@ protected:
 	*
 	* @param TrackIndex - Index of the track we want to update
 	*/
-	bool UpdateTrackSection(int32 TrackIndex);
+	bool UpdateTrackSection(int32 TrackIndex, float Time);
 
 	/**
 	* CreateTrackSection, Create/replace a track section.
diff --git a/Engine/Plugins/Experimental/GeometryCache/Source/GeometryCache/Classes/GeometryCacheTrack.h b/Engine/Plugins/Experimental/GeometryCache/Source/GeometryCache/Classes/GeometryCacheTrack.h
--- a/Engine/Plugins/Experimental/GeometryCache/Source/GeometryCache/Classes/GeometryCacheTrack.h
+++ b/Engine/Plugins/Experimental/GeometryCache/Source/GeometryCache/Classes/GeometryCacheTrack.h
@@ -125,6 +125,9 @@ class GEOMETRYCACHE_API UGeometryCacheTrack : public UObject
 	/** Return the hash of the mesh data of the track */
 	virtual uint64 GetHash() const { return 0; }
 
+	/** Update the current time of the track */
+	virtual void UpdateTime(float Time, bool bLooping) { }
+
 protected:
 
 	/** The duration of this track's animation. This is an open ended interval [0..Duration[.
diff --git a/Engine/Plugins/Experimental/GeometryCache/Source/GeometryCache/Private/GeometryCacheComponent.cpp b/Engine/Plugins/Experimental/GeometryCache/Source/GeometryCache/Private/GeometryCacheComponent.cpp
--- a/Engine/Plugins/Experimental/GeometryCache/Source/GeometryCache/Private/GeometryCacheComponent.cpp
+++ b/Engine/Plugins/Experimental/GeometryCache/Source/GeometryCache/Private/GeometryCacheComponent.cpp
@@ -110,7 +110,7 @@ void UGeometryCacheComponent::TickComponent(float DeltaTime, enum ELevelTick Tic
 		bool bUpdatedBoundsOrMatrix = false;
 		for (int32 TrackIndex = 0; TrackIndex < NumTracks; ++TrackIndex)
 		{
-			bUpdatedBoundsOrMatrix |= UpdateTrackSection(TrackIndex);
+			bUpdatedBoundsOrMatrix |= UpdateTrackSection(TrackIndex, GetAnimationTime());
 		}
 
 		if (bUpdatedBoundsOrMatrix)
@@ -169,7 +169,7 @@ void UGeometryCacheComponent::TickAtThisTime(const float Time, bool bInIsRunning
 		bool bUpdatedBoundsOrMatrix = false;
 		for (int32 TrackIndex = 0; TrackIndex < NumTracks; ++TrackIndex)
 		{
-			bUpdatedBoundsOrMatrix |= UpdateTrackSection(TrackIndex);
+			bUpdatedBoundsOrMatrix |= UpdateTrackSection(TrackIndex, Time);
 		}
 
 		if (bUpdatedBoundsOrMatrix)
@@ -268,21 +268,23 @@ void UGeometryCacheComponent::CreateTrackSection(int32 TrackIndex)
 		TrackSections.SetNum(TrackIndex + 1, false);
 	}
 
-	UpdateTrackSection(TrackIndex);
+	UpdateTrackSection(TrackIndex, GetAnimationTime());
 	MarkRenderStateDirty(); // Recreate scene proxy
 }
 
-bool UGeometryCacheComponent::UpdateTrackSection(int32 TrackIndex)
+bool UGeometryCacheComponent::UpdateTrackSection(int32 TrackIndex, float Time)
 {
 	checkf(TrackIndex < TrackSections.Num() && GeometryCache != nullptr && TrackIndex < GeometryCache->Tracks.Num(), TEXT("Invalid SectionIndex") );
 
 	UGeometryCacheTrack* Track = GeometryCache->Tracks[TrackIndex];
 	FTrackRenderData& UpdateSection = TrackSections[TrackIndex];
 
+	Track->UpdateTime(Time, bLooping);
+
 	FMatrix Matrix;
 	FBox TrackBounds;
-	const bool bUpdateMatrix = Track->UpdateMatrixData(GetAnimationTime(), bLooping, UpdateSection.MatrixSampleIndex, Matrix);
-	const bool bUpdateBounds = Track->UpdateBoundsData(GetAnimationTime(), bLooping, (PlayDirection < 0.0f) ? true : false, UpdateSection.BoundsSampleIndex, TrackBounds);
+	const bool bUpdateMatrix = Track->UpdateMatrixData(Time, bLooping, UpdateSection.MatrixSampleIndex, Matrix);
+	const bool bUpdateBounds = Track->UpdateBoundsData(Time, bLooping, (PlayDirection < 0.0f) ? true : false, UpdateSection.BoundsSampleIndex, TrackBounds);
 
 	if (bUpdateBounds)
 	{
diff --git a/Engine/Plugins/Experimental/GeometryCache/Source/GeometryCacheStreamer/Private/GeometryCacheStreamBase.cpp b/Engine/Plugins/Experimental/GeometryCache/Source/GeometryCacheStreamer/Private/GeometryCacheStreamBase.cpp
new file mode 100644
--- /dev/null
+++ b/Engine/Plugins/Experimental/GeometryCache/Source/GeometryCacheStreamer/Private/GeometryCacheStreamBase.cpp
@@ -0,0 +1,436 @@
+// Copyright Epic Games, Inc. All Rights Reserved.
+
+#include "GeometryCacheStreamBase.h"
+#include "Async/Async.h"
+#include "GeometryCacheMeshData.h"
+#include "Misc/ScopeLock.h"
+
+enum class EStreamReadRequestStatus
+{
+	Scheduled,
+	Completed,
+	Cancelled
+};
+
+struct FGeometryCacheStreamReadRequest
+{
+	FGeometryCacheMeshData* MeshData = nullptr;
+	int32 ReadIndex = 0;
+	int32 FrameIndex = 0;
+	EStreamReadRequestStatus Status = EStreamReadRequestStatus::Scheduled;
+};
+
+FGeometryCacheStreamBase::FGeometryCacheStreamBase(int32 ReadConcurrency, FGeometryCacheStreamDetails&& InDetails)
+: Details(InDetails)
+, CurrentFrameIndex(0)
+, MaxCachedFrames(0)
+, MaxCachedDuration(0.f)
+, MaxMemAllowed(TNumericLimits<float>::Max())
+, MemoryUsed(0.f)
+, bCancellationRequested(false)
+, bCacheNeedsUpdate(false)
+{
+	ensureMsgf(Details.SecondsPerFrame > 0.f, TEXT("GeometryCache stream was initialized with %f seconds per frame, which is invalid. Falling back to 1/24"), Details.SecondsPerFrame);
+	if (Details.SecondsPerFrame <= 0.f)
+	{
+		Details.SecondsPerFrame = 1.f / 24.f;
+	}
+
+	for (int32 Index = 0; Index < ReadConcurrency; ++Index)
+	{
+		// Populate the ReadIndices. Note that it used as a stack
+		ReadIndices.Push(Index);
+
+		// Populate pool of reusable ReadRequests
+		ReadRequestsPool.Add(new FGeometryCacheStreamReadRequest());
+	}
+}
+
+FGeometryCacheStreamBase::~FGeometryCacheStreamBase()
+{
+	CancelRequests();
+
+	// Delete all the cached MeshData
+	FReadScopeLock ReadLock(FramesAvailableLock);
+	for (const auto& Pair : FramesAvailable)
+	{
+		FGeometryCacheMeshData* MeshData = Pair.Value;
+		DecrementMemoryStat(*MeshData);
+		delete MeshData;
+	}
+
+	// And ReadRequests from the pool
+	for (int32 Index = 0; Index < ReadRequestsPool.Num(); ++Index)
+	{
+		delete ReadRequestsPool[Index];
+	}
+}
+
+int32 FGeometryCacheStreamBase::CancelRequests()
+{
+	TGuardValue<std::atomic<bool>, bool> CancellationRequested(bCancellationRequested, true);
+
+	// Clear the FramesNeeded to prevent scheduling further reads
+	FramesNeeded.Empty();
+
+	// Wait for all read requests to complete
+	TArray<int32> CompletedFrames;
+	while (FramesRequested.Num())
+	{
+		UpdateRequestStatus(CompletedFrames);
+		if (FramesRequested.Num())
+		{
+			FPlatformProcess::Sleep(0.01f);
+		}
+	}
+
+	return CompletedFrames.Num();
+}
+
+bool FGeometryCacheStreamBase::RequestFrameData()
+{
+	check(IsInGameThread());
+
+	if (FramesNeeded.Num() == 0)
+	{
+		return false;
+	}
+
+	FReadScopeLock ReadLock(FramesAvailableLock);
+
+	// Get the next frame index to read that is not already available and not in flight
+	int32 FrameIndex = FramesNeeded[0];
+	bool bFrameIndexValid = false;
+	while (FramesNeeded.Num() && !bFrameIndexValid)
+	{
+		FrameIndex = FramesNeeded[0];
+		bFrameIndexValid = true;
+		if (FramesAvailable.Contains(FrameIndex))
+		{
+			FramesNeeded.Remove(FrameIndex);
+			bFrameIndexValid = false;
+		}
+		for (const FGeometryCacheStreamReadRequest* Request : FramesRequested)
+		{
+			if (Request->FrameIndex == FrameIndex)
+			{
+				FramesNeeded.Remove(FrameIndex);
+				bFrameIndexValid = false;
+				break;
+			}
+		}
+	}
+
+	if (!bFrameIndexValid)
+	{
+		return false;
+	}
+
+	if (ReadIndices.Num() > 0)
+	{
+		// Get any ReadIndex available
+		const int32 ReadIndex = ReadIndices.Pop(false);
+
+		// Take the ReadRequest from the pool at ReadIndex and initialize it
+		FGeometryCacheStreamReadRequest*& ReadRequest = ReadRequestsPool[ReadIndex];
+		ReadRequest->FrameIndex = FrameIndex;
+		ReadRequest->ReadIndex = ReadIndex;
+		ReadRequest->MeshData = new FGeometryCacheMeshData;
+		ReadRequest->Status = EStreamReadRequestStatus::Scheduled;
+
+		// Change the frame status from needed to requested
+		FramesNeeded.Remove(FrameIndex);
+		FramesRequested.Add(ReadRequest);
+
+		// Schedule asynchronous read of the MeshData
+		Async(
+#if WITH_EDITOR
+			EAsyncExecution::LargeThreadPool,
+#else
+			EAsyncExecution::ThreadPool,
+#endif // WITH_EDITOR
+			[this, ReadRequest]()
+			{
+				if (!bCancellationRequested)
+				{
+					checkSlow(ReadRequest->MeshData);
+					GetMeshData(ReadRequest->FrameIndex, ReadRequest->ReadIndex, *ReadRequest->MeshData);
+					ReadRequest->Status = EStreamReadRequestStatus::Completed;
+				}
+				else
+				{
+					ReadRequest->Status = EStreamReadRequestStatus::Cancelled;
+				}
+			});
+
+		return true;
+	}
+	return false;
+}
+
+void FGeometryCacheStreamBase::UpdateRequestStatus(TArray<int32>& OutFramesCompleted)
+{
+	check(IsInGameThread());
+
+	// Check if the cache needs to be updated either because the frame has advanced or because there's a new memory limit
+	if (bCacheNeedsUpdate)
+	{
+		bCacheNeedsUpdate = false;
+		UpdateFramesNeeded(CurrentFrameIndex, MaxCachedFrames);
+
+		// Remove the unneeded frames, those that are available but don't need to be cached
+		TArray<int32> UnneededFrames;
+		{
+			FReadScopeLock ReadLock(FramesAvailableLock);
+			for (const auto& Pair : FramesAvailable)
+			{
+				int32 FrameIndex = Pair.Key;
+				if (!FramesToBeCached.Contains(FrameIndex))
+				{
+					FGeometryCacheMeshData* MeshData = Pair.Value;
+					DecrementMemoryStat(*MeshData);
+					delete MeshData;
+
+					UnneededFrames.Add(FrameIndex);
+				}
+			}
+		}
+
+		{
+			FWriteScopeLock WriteLock(FramesAvailableLock);
+			for (int32 FrameIndex : UnneededFrames)
+			{
+				FramesAvailable.Remove(FrameIndex);
+			}
+		}
+	}
+
+	FWriteScopeLock WriteLock(FramesAvailableLock);
+
+	// Check the completion status of the read requests in progress
+	TArray<FGeometryCacheStreamReadRequest*> CompletedRequests;
+	for (FGeometryCacheStreamReadRequest* ReadRequest : FramesRequested)
+	{
+		// A cancelled read is still considered completed, it just hasn't read any data
+		if (ReadRequest->Status == EStreamReadRequestStatus::Completed ||
+			ReadRequest->Status == EStreamReadRequestStatus::Cancelled)
+		{
+			// Queue for removal after iterating
+			CompletedRequests.Add(ReadRequest);
+
+			// A cancelled read has an allocated mesh data that needs to be deleted
+			bool bDeleteMeshData = ReadRequest->Status == EStreamReadRequestStatus::Cancelled;
+			if (ReadRequest->Status == EStreamReadRequestStatus::Completed)
+			{
+				checkSlow(ReadRequest->MeshData);
+				FResourceSizeEx ResSize;
+				ReadRequest->MeshData->GetResourceSizeEx(ResSize);
+				float FrameDataSize = float(ResSize.GetTotalMemoryBytes()) / (1024 * 1024);
+
+				if (MemoryUsed + FrameDataSize < MaxMemAllowed)
+				{
+					if (!FramesAvailable.Contains(ReadRequest->FrameIndex))
+					{
+						// Cache result of read for retrieval later
+						FramesAvailable.Add(ReadRequest->FrameIndex, ReadRequest->MeshData);
+						IncrementMemoryStat(*ReadRequest->MeshData);
+					}
+					else
+					{
+						// The requested frame was already available, just delete it
+						bDeleteMeshData = true;
+					}
+				}
+				else
+				{
+					// The frame doesn't fit in the allowed memory budget so delete it
+					// It should be requested again later if it's still needed
+					bDeleteMeshData = true;
+				}
+			}
+
+			// Push back the ReadIndex for reuse
+			ReadIndices.Push(ReadRequest->ReadIndex);
+
+			// Output the completed frame
+			OutFramesCompleted.Add(ReadRequest->FrameIndex);
+
+			if (bDeleteMeshData)
+			{
+				delete ReadRequest->MeshData;
+				ReadRequest->MeshData = nullptr;
+			}
+		}
+	}
+
+	for (FGeometryCacheStreamReadRequest* ReadRequest : CompletedRequests)
+	{
+		FramesRequested.Remove(ReadRequest);
+	}
+}
+
+void FGeometryCacheStreamBase::UpdateFramesNeeded(int32 StartFrameIndex, int32 NumFrames)
+{
+	FramesToBeCached.Empty(NumFrames);
+	FramesNeeded.Empty(NumFrames);
+
+	const int32 StartIndex = Details.StartFrameIndex;
+	const int32 EndIndex = Details.EndFrameIndex;
+
+	// FramesToBeCached are the frame indices that are required for playback, available or not,
+	// while FramesNeeded are the frame indices that are not available yet so they need to be read
+	auto AddFrameIndex = [this, &NumFrames](int32 FrameIndex)
+	{
+		FramesToBeCached.Add(FrameIndex);
+		if (!FramesAvailable.Contains(FrameIndex))
+		{
+			FramesNeeded.Add(FrameIndex);
+		}
+		--NumFrames;
+	};
+
+	FReadScopeLock ReadLock(FramesAvailableLock);
+
+	StartFrameIndex = FMath::Clamp(StartFrameIndex, StartIndex, EndIndex);
+
+	// Also reserve space for the frame before start since playback is double-buffered
+	int PreviousFrameIndex = FMath::Clamp(StartFrameIndex - 1, StartIndex, EndIndex);
+	if (PreviousFrameIndex != StartFrameIndex)
+	{
+		--NumFrames;
+	}
+
+	// Populate the list of frame indices from given StartFrameIndex up to NumFrames or EndIndex
+	for (int32 Index = StartFrameIndex; NumFrames > 0 && Index < EndIndex; ++Index)
+	{
+		AddFrameIndex(Index);
+	}
+
+	// End of the range might have been reached before the requested NumFrames so add the remaining frames starting from StartIndex
+	for (int32 Index = StartIndex; NumFrames > 0 && Index < PreviousFrameIndex; ++Index)
+	{
+		AddFrameIndex(Index);
+	}
+
+	// Frame before start is added at the end to preserve the priority of the other frames
+	if (PreviousFrameIndex != StartFrameIndex)
+	{
+		AddFrameIndex(PreviousFrameIndex);
+	}
+}
+
+void FGeometryCacheStreamBase::Prefetch(int32 StartFrameIndex, int32 NumFrames)
+{
+	const int32 MaxNumFrames = Details.NumFrames;
+
+	// Validate the number of frames to be loaded
+	if (NumFrames == 0)
+	{
+		// If no value specified, load the whole stream
+		NumFrames = MaxNumFrames;
+	}
+	else
+	{
+		NumFrames = FMath::Clamp(NumFrames, 1, MaxNumFrames);
+	}
+
+	MaxCachedFrames = NumFrames;
+
+	UpdateFramesNeeded(StartFrameIndex, NumFrames);
+
+	if (FramesNeeded.Num() > 0)
+	{
+		// Force the first frame to be loaded and ready for retrieval
+		LoadFrameData(FramesNeeded[0]);
+		FramesNeeded.RemoveAt(0);
+	}
+}
+
+uint32 FGeometryCacheStreamBase::GetNumFramesNeeded()
+{
+	return FramesNeeded.Num();
+}
+
+bool FGeometryCacheStreamBase::GetFrameData(int32 FrameIndex, FGeometryCacheMeshData& OutMeshData)
+{
+	// This function can be called from the render thread
+	FReadScopeLock ReadLock(FramesAvailableLock);
+	if (FGeometryCacheMeshData** MeshDataPtr = FramesAvailable.Find(FrameIndex))
+	{
+		OutMeshData = **MeshDataPtr;
+		return true;
+	}
+	return false;
+}
+
+void FGeometryCacheStreamBase::LoadFrameData(int32 FrameIndex)
+{
+	check(IsInGameThread());
+
+	FWriteScopeLock WriteLock(FramesAvailableLock);
+	if (FramesAvailable.Contains(FrameIndex))
+	{
+		return;
+	}
+
+	FGeometryCacheMeshData* MeshData = new FGeometryCacheMeshData;
+	GetMeshData(FrameIndex, 0, *MeshData);
+	FramesAvailable.Add(FrameIndex, MeshData);
+	IncrementMemoryStat(*MeshData);
+}
+
+const FGeometryCacheStreamStats& FGeometryCacheStreamBase::GetStreamStats() const
+{
+	check(IsInGameThread());
+
+	const int32 NumFrames = FramesAvailable.Num();
+	const float Secs = Details.SecondsPerFrame * NumFrames;
+
+	Stats.NumCachedFrames = NumFrames;
+	Stats.CachedDuration = Secs;
+	Stats.MemoryUsed = MemoryUsed;
+	Stats.AverageBitrate = MemoryUsed / Secs;
+	return Stats;
+}
+
+void FGeometryCacheStreamBase::SetLimits(float InMaxMemoryAllowed, float InMaxCachedDuration)
+{
+	check(IsInGameThread());
+
+	if (!FMath::IsNearlyEqual(InMaxMemoryAllowed, MaxMemAllowed) ||
+		!FMath::IsNearlyEqual(InMaxCachedDuration, MaxCachedDuration))
+	{
+		bCacheNeedsUpdate = true;
+		MaxMemAllowed = InMaxMemoryAllowed;
+
+		MaxCachedDuration = FMath::Min(InMaxCachedDuration, Details.Duration);
+		MaxCachedFrames = FMath::Min(FMath::FloorToInt(MaxCachedDuration / Details.SecondsPerFrame), Details.NumFrames);
+	}
+}
+
+void FGeometryCacheStreamBase::IncrementMemoryStat(const FGeometryCacheMeshData& MeshData)
+{
+	FResourceSizeEx ResSize;
+	MeshData.GetResourceSizeEx(ResSize);
+
+	float SizeInBytes = ResSize.GetTotalMemoryBytes();
+	MemoryUsed += SizeInBytes / (1024 * 1024);
+}
+
+void FGeometryCacheStreamBase::DecrementMemoryStat(const FGeometryCacheMeshData& MeshData)
+{
+	FResourceSizeEx ResSize;
+	MeshData.GetResourceSizeEx(ResSize);
+
+	float SizeInBytes = ResSize.GetTotalMemoryBytes();
+	MemoryUsed -= SizeInBytes / (1024 * 1024);
+}
+
+void FGeometryCacheStreamBase::UpdateCurrentFrameIndex(int32 FrameIndex)
+{
+	if (FrameIndex != CurrentFrameIndex)
+	{
+		CurrentFrameIndex = FrameIndex;
+		bCacheNeedsUpdate = true;
+	}
+}
diff --git a/Engine/Plugins/Experimental/GeometryCache/Source/GeometryCacheStreamer/Private/GeometryCacheStreamer.cpp b/Engine/Plugins/Experimental/GeometryCache/Source/GeometryCacheStreamer/Private/GeometryCacheStreamer.cpp
--- a/Engine/Plugins/Experimental/GeometryCache/Source/GeometryCacheStreamer/Private/GeometryCacheStreamer.cpp
+++ b/Engine/Plugins/Experimental/GeometryCache/Source/GeometryCacheStreamer/Private/GeometryCacheStreamer.cpp
@@ -14,6 +14,8 @@
 #define LOCTEXT_NAMESPACE "GeometryCacheStreamer"
 
 DECLARE_DWORD_COUNTER_STAT(TEXT("GeometryCache Streams: Count"), STAT_GeometryCacheStream_Count, STATGROUP_GeometryCache);
+DECLARE_FLOAT_COUNTER_STAT(TEXT("GeometryCache Streams: Look-Ahead (seconds)"), STAT_GeometryCacheStream_LookAhead, STATGROUP_GeometryCache);
+DECLARE_FLOAT_COUNTER_STAT(TEXT("GeometryCache Streams: Total Bitrate (MB/s)"), STAT_GeometryCacheStream_Bitrate, STATGROUP_GeometryCache);
 DECLARE_MEMORY_STAT(TEXT("GeometryCache Streams: Memory Used"), STAT_GeometryCacheStream_MemoryUsed, STATGROUP_GeometryCache);
 
 static bool GShowGeometryCacheStreamerNotification = true;
@@ -196,14 +198,18 @@ void FGeometryCacheStreamer::BalanceStreams()
 	}
 
 	SET_MEMORY_STAT(STAT_GeometryCacheStream_MemoryUsed, AggregatedStats.MemoryUsed * 1024 * 1024);
+	SET_FLOAT_STAT(STAT_GeometryCacheStream_LookAhead, AggregatedStats.CachedDuration);
+	SET_FLOAT_STAT(STAT_GeometryCacheStream_Bitrate, AggregatedStats.AverageBitrate);
 
 	const UGeometryCacheStreamerSettings* Settings = GetDefault<UGeometryCacheStreamerSettings>();
 	const float MemoryLimit = Settings->MaxMemoryAllowed;
 	const float DurationLimit = Settings->LookAheadBuffer;
 
 	// Apply memory limits if over budget
-	// Might want to also update the limits if under budget
-	if (AggregatedStats.MemoryUsed > MemoryLimit)
+	// Re-balance when under budget only when it's idle as it can naturally go under if it is starving
+	bool bIsOverBudget = AggregatedStats.MemoryUsed > MemoryLimit || AggregatedStats.CachedDuration > DurationLimit;
+	bool bIsUnderBudget = AggregatedStats.MemoryUsed < MemoryLimit * 0.75f || AggregatedStats.CachedDuration < DurationLimit;
+	if (bIsOverBudget || (bIsUnderBudget && NumReads == 0))
 	{
 		// Allocate per-stream budget based on the duration computed from the total bitrate and given memory limit
 		const float AvgDuration = MemoryLimit / AggregatedStats.AverageBitrate;
@@ -233,7 +239,6 @@ void FGeometryCacheStreamer::UnregisterTrack(UGeometryCacheTrack* AbcTrack)
 		int32 NumRequests = (*Stream)->CancelRequests();
 		NumReads -= NumRequests;
 
-		delete *Stream;
 		TracksToStreams.Remove(AbcTrack);
 	}
 }
diff --git a/Engine/Plugins/Experimental/GeometryCache/Source/GeometryCacheStreamer/Public/GeometryCacheStreamBase.h b/Engine/Plugins/Experimental/GeometryCache/Source/GeometryCacheStreamer/Public/GeometryCacheStreamBase.h
new file mode 100644
--- /dev/null
+++ b/Engine/Plugins/Experimental/GeometryCache/Source/GeometryCacheStreamer/Public/GeometryCacheStreamBase.h
@@ -0,0 +1,78 @@
+// Copyright Epic Games, Inc. All Rights Reserved.
+
+#pragma once
+
+#include "IGeometryCacheStream.h"
+#include <atomic>
+
+struct FGeometryCacheStreamReadRequest;
+
+/* Details about the animation to be streamed */
+struct FGeometryCacheStreamDetails
+{
+	int32 NumFrames = 0;
+	float Duration = 0.f;
+	float SecondsPerFrame = 1.f / 24.f;
+	int32 StartFrameIndex = 0;
+	int32 EndFrameIndex = 0;
+};
+
+/**
+ * Base class for GeometryCache stream for use with the GeometryCacheStreamer
+ * Besides implementing the basic functionalities expected of the stream,
+ * it implements basic memory statistics and management for use by the streamer.
+ * Derived classes need to implement a way to retrieve the mesh data for a frame
+ * through GetMeshData.
+ */
+class GEOMETRYCACHESTREAMER_API FGeometryCacheStreamBase : public IGeometryCacheStream
+{
+public:
+	FGeometryCacheStreamBase(int32 ReadConcurrency, FGeometryCacheStreamDetails&& Details);
+	virtual ~FGeometryCacheStreamBase();
+
+	//~ Begin IGeometryCacheStream Interface
+	virtual void Prefetch(int32 StartFrameIndex, int32 NumFrames = 0) override;
+	virtual uint32 GetNumFramesNeeded() override;
+	virtual bool RequestFrameData() override;
+	virtual void UpdateRequestStatus(TArray<int32>& OutFramesCompleted) override;
+	virtual bool GetFrameData(int32 FrameIndex, FGeometryCacheMeshData& OutMeshData) override;
+	virtual int32 CancelRequests() override;
+	virtual const FGeometryCacheStreamStats& GetStreamStats() const override;
+	virtual void SetLimits(float MaxMemoryAllowed, float MaxCachedDuration) override;
+	//~ End IGeometryCacheStream Interface
+
+	/* Updates the current position in the stream */
+	void UpdateCurrentFrameIndex(int32 FrameIndex);
+
+protected:
+
+	/* Derived class must provide a way to get the mesh data for the given FrameIndex */
+	virtual void GetMeshData(int32 FrameIndex, int32 ReadConcurrencyIndex, FGeometryCacheMeshData& OutMeshData) = 0;
+
+	void LoadFrameData(int32 FrameIndex);
+	void UpdateFramesNeeded(int32 StartIndex, int32 NumFrames);
+	void IncrementMemoryStat(const FGeometryCacheMeshData& MeshData);
+	void DecrementMemoryStat(const FGeometryCacheMeshData& MeshData);
+
+	TArray<int32> ReadIndices;
+	TArray<FGeometryCacheStreamReadRequest*> ReadRequestsPool;
+
+	TArray<int32> FramesNeeded;
+	TArray<int32> FramesToBeCached;
+	TArray<FGeometryCacheStreamReadRequest*> FramesRequested;
+
+	using FFrameIndexToMeshData = TMap<int32, FGeometryCacheMeshData*>;
+	FFrameIndexToMeshData FramesAvailable;
+	FRWLock FramesAvailableLock;
+
+	FGeometryCacheStreamDetails Details;
+	mutable FGeometryCacheStreamStats Stats;
+	int32 CurrentFrameIndex;
+	int32 MaxCachedFrames;
+	float MaxCachedDuration;
+	float MaxMemAllowed;
+	float MemoryUsed;
+
+	std::atomic<bool> bCancellationRequested;
+	bool bCacheNeedsUpdate;
+};
diff --git a/Engine/Plugins/Experimental/GeometryCacheAbcFile/Source/GeometryCacheAbcFile/Private/GeometryCacheAbcStream.cpp b/Engine/Plugins/Experimental/GeometryCacheAbcFile/Source/GeometryCacheAbcFile/Private/GeometryCacheAbcStream.cpp
--- a/Engine/Plugins/Experimental/GeometryCacheAbcFile/Source/GeometryCacheAbcFile/Private/GeometryCacheAbcStream.cpp
+++ b/Engine/Plugins/Experimental/GeometryCacheAbcFile/Source/GeometryCacheAbcFile/Private/GeometryCacheAbcStream.cpp
@@ -3,14 +3,11 @@
 #include "GeometryCacheAbcStream.h"
 #include "AbcFile.h"
 #include "AbcUtilities.h"
-#include "Async/Async.h"
 #include "DerivedDataCacheInterface.h"
 #include "GeometryCacheMeshData.h"
-#include "GeometryCacheStreamerSettings.h"
 #include "GeometryCacheTrackAbcFile.h"
-#include "Math/NumericLimits.h"
+#include "HAL/IConsoleManager.h"
 #include "Misc/CoreMisc.h"
-#include "Misc/ScopeLock.h"
 
 static bool GAbcStreamCacheInDDC = true;
 static FAutoConsoleVariableRef CVarAbcStreamCacheInDDC(
@@ -69,67 +66,17 @@ public:
 };
 
 FGeometryCacheAbcStream::FGeometryCacheAbcStream(UGeometryCacheTrackAbcFile* InAbcTrack)
-: AbcTrack(InAbcTrack)
-, bCancellationRequested(false)
-, LastAccessedFrameIndex(0)
-, MaxCachedFrames(0)
-, MaxCachedDuration(0.f)
-, MaxMemAllowed(TNumericLimits<float>::Max())
-, MemoryUsed(0.f)
-, bCacheNeedsUpdate(false)
+: FGeometryCacheStreamBase(
+	kAbcReadConcurrency,
+	FGeometryCacheStreamDetails{
+		InAbcTrack->GetAbcFile().GetImportNumFrames(),
+		InAbcTrack->GetAbcFile().GetImportLength(),
+		InAbcTrack->GetAbcFile().GetSecondsPerFrame(),
+		InAbcTrack->GetAbcFile().GetStartFrameIndex(),
+		InAbcTrack->GetAbcFile().GetEndFrameIndex()})
+, AbcTrack(InAbcTrack)
+, Hash(InAbcTrack->GetAbcTrackHash())
 {
-	SecondsPerFrame = AbcTrack->GetAbcFile().GetSecondsPerFrame();
-	Hash = AbcTrack->GetAbcTrackHash();
-
-	for (int32 Index = 0; Index < kAbcReadConcurrency; ++Index)
-	{
-		// Populate the ReadIndices. Note that it used as a stack
-		ReadIndices.Push(Index);
-
-		// Populate pool of reusable ReadRequests
-		ReadRequestsPool.Add(new FGeometryCacheAbcStreamReadRequest());
-	}
-}
-
-FGeometryCacheAbcStream::~FGeometryCacheAbcStream()
-{
-	CancelRequests();
-
-	// Delete all the cached MeshData
-	FReadScopeLock ReadLock(FramesAvailableLock);
-	for (const auto& Pair : FramesAvailable)
-	{
-		FGeometryCacheMeshData* MeshData = Pair.Value;
-		DecrementMemoryStat(*MeshData);
-		delete MeshData;
-	}
-
-	// And ReadRequests from the pool
-	for (int32 Index = 0; Index < ReadRequestsPool.Num(); ++Index)
-	{
-		delete ReadRequestsPool[Index];
-	}
-}
-
-int32 FGeometryCacheAbcStream::CancelRequests()
-{
-	TGuardValue<std::atomic<bool>, bool> CancellationRequested(bCancellationRequested, true);
-
-	// Clear the FramesNeeded to prevent scheduling further reads
-	FramesNeeded.Empty();
-
-	// Wait for all read requests to complete
-	TArray<int32> CompletedFrames;
-	while (FramesRequested.Num())
-	{
-		UpdateRequestStatus(CompletedFrames);
-		if (FramesRequested.Num())
-		{
-			FPlatformProcess::Sleep(0.01f);
-		}
-	}
-
-	return CompletedFrames.Num();
 }
 
 void FGeometryCacheAbcStream::GetMeshData(int32 FrameIndex, int32 ConcurrencyIndex, FGeometryCacheMeshData& MeshData)
@@ -162,349 +109,3 @@ void FGeometryCacheAbcStream::GetMeshData(int32 FrameIndex, int32 ConcurrencyInd
 		FAbcUtilities::GetFrameMeshData(AbcTrack->GetAbcFile(), FrameIndex, MeshData, ConcurrencyIndex);
 	}
 }
-
-bool FGeometryCacheAbcStream::RequestFrameData()
-{
-	check(IsInGameThread());
-
-	if (FramesNeeded.Num() == 0)
-	{
-		return false;
-	}
-
-	FReadScopeLock ReadLock(FramesAvailableLock);
-
-	// Get the next frame index to read that is not already available and not in flight
-	int32 FrameIndex = FramesNeeded[0];
-	bool bFrameIndexValid = false;
-	while (FramesNeeded.Num() && !bFrameIndexValid)
-	{
-		FrameIndex = FramesNeeded[0];
-		bFrameIndexValid = true;
-		if (FramesAvailable.Contains(FrameIndex))
-		{
-			FramesNeeded.Remove(FrameIndex);
-			bFrameIndexValid = false;
-		}
-		for (const FGeometryCacheAbcStreamReadRequest* Request : FramesRequested)
-		{
-			if (Request->FrameIndex == FrameIndex)
-			{
-				FramesNeeded.Remove(FrameIndex);
-				bFrameIndexValid = false;
-				break;
-			}
-		}
-	}
-
-	if (!bFrameIndexValid)
-	{
-		return false;
-	}
-
-	if (ReadIndices.Num() > 0)
-	{
-		// Get any ReadIndex available
-		const int32 ReadIndex = ReadIndices.Pop(false);
-
-		// Take the ReadRequest from the pool at ReadIndex and initialize it
-		FGeometryCacheAbcStreamReadRequest*& ReadRequest = ReadRequestsPool[ReadIndex];
-		ReadRequest->FrameIndex = FrameIndex;
-		ReadRequest->ReadIndex = ReadIndex;
-		ReadRequest->MeshData = new FGeometryCacheMeshData;
-		ReadRequest->Status = EAbcStreamReadRequestStatus::Scheduled;
-
-		// Change the frame status from needed to requested
-		FramesNeeded.Remove(FrameIndex);
-		FramesRequested.Add(ReadRequest);
-
-		// Schedule asynchronous read of the MeshData
-		Async(EAsyncExecution::LargeThreadPool,
-			[this, ReadRequest]()
-			{
-				if (!bCancellationRequested)
-				{
-					checkSlow(ReadRequest->MeshData);
-					GetMeshData(ReadRequest->FrameIndex, ReadRequest->ReadIndex, *ReadRequest->MeshData);
-					ReadRequest->Status = EAbcStreamReadRequestStatus::Completed;
-				}
-				else
-				{
-					ReadRequest->Status = EAbcStreamReadRequestStatus::Cancelled;
-				}
-			});
-
-		return true;
-	}
-	return false;
-}
-
-void FGeometryCacheAbcStream::UpdateRequestStatus(TArray<int32>& OutFramesCompleted)
-{
-	check(IsInGameThread());
-
-	// Check if the cache needs to be updated either because the frame has advanced or because there's a new memory limit
-	if (bCacheNeedsUpdate)
-	{
-		bCacheNeedsUpdate = false;
-		UpdateFramesNeeded(LastAccessedFrameIndex, MaxCachedFrames);
-
-		// Remove the unneeded frames, those that are available but don't need to be cached
-		TArray<int32> UnneededFrames;
-		{
-			FReadScopeLock ReadLock(FramesAvailableLock);
-			for (const auto& Pair : FramesAvailable)
-			{
-				int32 FrameIndex = Pair.Key;
-				if (!FramesToBeCached.Contains(FrameIndex))
-				{
-					FGeometryCacheMeshData* MeshData = Pair.Value;
-					DecrementMemoryStat(*MeshData);
-					delete MeshData;
-
-					UnneededFrames.Add(FrameIndex);
-				}
-			}
-		}
-
-		{
-			FWriteScopeLock WriteLock(FramesAvailableLock);
-			for (int32 FrameIndex : UnneededFrames)
-			{
-				FramesAvailable.Remove(FrameIndex);
-			}
-		}
-	}
-
-	FWriteScopeLock WriteLock(FramesAvailableLock);
-
-	// Check the completion status of the read requests in progress
-	TArray<FGeometryCacheAbcStreamReadRequest*> CompletedRequests;
-	for (FGeometryCacheAbcStreamReadRequest* ReadRequest : FramesRequested)
-	{
-		// A cancelled read is still considered completed, it just hasn't read any data
-		if (ReadRequest->Status == EAbcStreamReadRequestStatus::Completed ||
-			ReadRequest->Status == EAbcStreamReadRequestStatus::Cancelled)
-		{
-			// Queue for removal after iterating
-			CompletedRequests.Add(ReadRequest);
-
-			// A cancelled read has an allocated mesh data that needs to be deleted
-			bool bDeleteMeshData = ReadRequest->Status == EAbcStreamReadRequestStatus::Cancelled;
-			if (ReadRequest->Status == EAbcStreamReadRequestStatus::Completed)
-			{
-				checkSlow(ReadRequest->MeshData);
-				FResourceSizeEx ResSize;
-				ReadRequest->MeshData->GetResourceSizeEx(ResSize);
-				float FrameDataSize = float(ResSize.GetTotalMemoryBytes()) / (1024 * 1024);
-
-				if (MemoryUsed + FrameDataSize < MaxMemAllowed)
-				{
-					if (!FramesAvailable.Contains(ReadRequest->FrameIndex))
-					{
-						// Cache result of read for retrieval later
-						FramesAvailable.Add(ReadRequest->FrameIndex, ReadRequest->MeshData);
-						IncrementMemoryStat(*ReadRequest->MeshData);
-					}
-					else
-					{
-						// The requested frame was already available, just delete it
-						bDeleteMeshData = true;
-					}
-				}
-				else
-				{
-					// The frame doesn't fit in the allowed memory budget so delete it
-					// It should be requested again later if it's still needed
-					bDeleteMeshData = true;
-				}
-			}
-
-			// Push back the ReadIndex for reuse
-			ReadIndices.Push(ReadRequest->ReadIndex);
-
-			// Output the completed frame
-			OutFramesCompleted.Add(ReadRequest->FrameIndex);
-
-			if (bDeleteMeshData)
-			{
-				delete ReadRequest->MeshData;
-				ReadRequest->MeshData = nullptr;
-			}
-		}
-	}
-
-	for (FGeometryCacheAbcStreamReadRequest* ReadRequest : CompletedRequests)
-	{
-		FramesRequested.Remove(ReadRequest);
-	}
-}
-
-void FGeometryCacheAbcStream::UpdateFramesNeeded(int32 StartFrameIndex, int32 NumFrames)
-{
-	FramesToBeCached.Empty(NumFrames);
-	FramesNeeded.Empty(NumFrames);
-
-	const FAbcFile& AbcFile = AbcTrack->GetAbcFile();
-	const int32 StartIndex = AbcFile.GetStartFrameIndex();
-	const int32 EndIndex = AbcFile.GetEndFrameIndex();
-
-	// FramesToBeCached are the frame indices that are required for playback, available or not,
-	// while FramesNeeded are the frame indices that are not available yet so they need to be read
-	auto AddFrameIndex = [this, &NumFrames](int32 FrameIndex)
-	{
-		FramesToBeCached.Add(FrameIndex);
-		if (!FramesAvailable.Contains(FrameIndex))
-		{
-			FramesNeeded.Add(FrameIndex);
-		}
-		--NumFrames;
-	};
-
-	FReadScopeLock ReadLock(FramesAvailableLock);
-
-	StartFrameIndex = FMath::Clamp(StartFrameIndex, StartIndex, EndIndex);
-
-	// Also reserve space for the frame before start since playback is double-buffered
-	int PreviousFrameIndex = FMath::Clamp(StartFrameIndex - 1, StartIndex, EndIndex);
-	if (PreviousFrameIndex != StartFrameIndex)
-	{
-		--NumFrames;
-	}
-
-	// Populate the list of frame indices from given StartFrameIndex up to NumFrames or EndIndex
-	for (int32 Index = StartFrameIndex; NumFrames > 0 && Index < EndIndex; ++Index)
-	{
-		AddFrameIndex(Index);
-	}
-
-	// End of the range might have been reached before the requested NumFrames so add the remaining frames starting from StartIndex
-	for (int32 Index = StartIndex; NumFrames > 0 && Index < PreviousFrameIndex; ++Index)
-	{
-		AddFrameIndex(Index);
-	}
-
-	// Frame before start is added at the end to preserve the priority of the other frames
-	if (PreviousFrameIndex != StartFrameIndex)
-	{
-		AddFrameIndex(PreviousFrameIndex);
-	}
-}
-
-void FGeometryCacheAbcStream::Prefetch(int32 StartFrameIndex, int32 NumFrames)
-{
-	const FAbcFile& AbcFile = AbcTrack->GetAbcFile();
-	const int32 MaxNumFrames = AbcFile.GetImportNumFrames();
-
-	// Validate the number of frames to be loaded
-	if (NumFrames == 0)
-	{
-		// If no value specified, load the whole stream
-		NumFrames = MaxNumFrames;
-	}
-	else
-	{
-		NumFrames = FMath::Clamp(NumFrames, 1, MaxNumFrames);
-	}
-
-	MaxCachedFrames = NumFrames;
-
-	UpdateFramesNeeded(StartFrameIndex, NumFrames);
-
-	if (FramesNeeded.Num() > 0)
-	{
-		// Force the first frame to be loaded and ready for retrieval
-		LoadFrameData(FramesNeeded[0]);
-		FramesNeeded.RemoveAt(0);
-	}
-}
-
-uint32 FGeometryCacheAbcStream::GetNumFramesNeeded()
-{
-	return FramesNeeded.Num();
-}
-
-bool FGeometryCacheAbcStream::GetFrameData(int32 FrameIndex, FGeometryCacheMeshData& OutMeshData)
-{
-	// This function can be called from the render thread
-	int32 PrevLastAccessedFrameIndex = LastAccessedFrameIndex.exchange(FrameIndex);
-	if (PrevLastAccessedFrameIndex != FrameIndex)
-	{
-		bCacheNeedsUpdate = true;
-	}
-
-	{
-		FReadScopeLock ReadLock(FramesAvailableLock);
-		if (FGeometryCacheMeshData** MeshDataPtr = FramesAvailable.Find(FrameIndex))
-		{
-			OutMeshData = **MeshDataPtr;
-			return true;
-		}
-	}
-
-	return false;
-}
-
-void FGeometryCacheAbcStream::LoadFrameData(int32 FrameIndex)
-{
-	check(IsInGameThread());
-
-	FWriteScopeLock WriteLock(FramesAvailableLock);
-	if (FramesAvailable.Contains(FrameIndex))
-	{
-		return;
-	}
-
-	FGeometryCacheMeshData* MeshData = new FGeometryCacheMeshData;
-	GetMeshData(FrameIndex, 0, *MeshData);
-	FramesAvailable.Add(FrameIndex, MeshData);
-	IncrementMemoryStat(*MeshData);
-}
-
-const FGeometryCacheStreamStats& FGeometryCacheAbcStream::GetStreamStats() const
-{
-	check(IsInGameThread());
-
-	const int32 NumFrames = FramesAvailable.Num();
-	const float Secs = SecondsPerFrame * NumFrames;
-
-	Stats.NumCachedFrames = NumFrames;
-	Stats.CachedDuration = Secs;
-	Stats.MemoryUsed = MemoryUsed;
-	Stats.AverageBitrate = MemoryUsed / Secs;
-	return Stats;
-}
-
-void FGeometryCacheAbcStream::SetLimits(float InMaxMemoryAllowed, float InMaxCachedDuration)
-{
-	check(IsInGameThread());
-
-	if (InMaxMemoryAllowed != MaxMemAllowed)
-	{
-		if (InMaxMemoryAllowed < MaxMemAllowed)
-		{
-			bCacheNeedsUpdate = true;
-		}
-		MaxMemAllowed = InMaxMemoryAllowed;
-		MaxCachedDuration = FMath::Min(InMaxCachedDuration, AbcTrack->GetAbcFile().GetImportLength());
-		MaxCachedFrames = FMath::Min(FMath::CeilToInt(MaxCachedDuration / SecondsPerFrame), AbcTrack->GetAbcFile().GetImportNumFrames());
-	}
-}
-
-void FGeometryCacheAbcStream::IncrementMemoryStat(const FGeometryCacheMeshData& MeshData)
-{
-	FResourceSizeEx ResSize;
-	MeshData.GetResourceSizeEx(ResSize);
-
-	float SizeInBytes = ResSize.GetTotalMemoryBytes();
-	MemoryUsed += SizeInBytes / (1024 * 1024);
-}
-
-void FGeometryCacheAbcStream::DecrementMemoryStat(const FGeometryCacheMeshData& MeshData)
-{
-	FResourceSizeEx ResSize;
-	MeshData.GetResourceSizeEx(ResSize);
-
-	float SizeInBytes = ResSize.GetTotalMemoryBytes();
-	MemoryUsed -= SizeInBytes / (1024 * 1024);
-}
diff --git a/Engine/Plugins/Experimental/GeometryCacheAbcFile/Source/GeometryCacheAbcFile/Private/GeometryCacheAbcStream.h b/Engine/Plugins/Experimental/GeometryCacheAbcFile/Source/GeometryCacheAbcFile/Private/GeometryCacheAbcStream.h
--- a/Engine/Plugins/Experimental/GeometryCacheAbcFile/Source/GeometryCacheAbcFile/Private/GeometryCacheAbcStream.h
+++ b/Engine/Plugins/Experimental/GeometryCacheAbcFile/Source/GeometryCacheAbcFile/Private/GeometryCacheAbcStream.h
@@ -2,59 +2,20 @@
 
 #pragma once
 
-#include "IGeometryCacheStream.h"
-#include <atomic>
+#include "GeometryCacheStreamBase.h"
 
-struct FGeometryCacheAbcStreamReadRequest;
 class UGeometryCacheTrackAbcFile;
 
-class FGeometryCacheAbcStream : public IGeometryCacheStream
+class FGeometryCacheAbcStream : public FGeometryCacheStreamBase
 {
 public:
 	FGeometryCacheAbcStream(UGeometryCacheTrackAbcFile* InAbcTrack);
-	virtual ~FGeometryCacheAbcStream();
 
+protected:
 	//~ Begin IGeometryCacheStream Interface
-	virtual void Prefetch(int32 StartFrameIndex, int32 NumFrames = 0) override;
-	virtual uint32 GetNumFramesNeeded() override;
-	virtual bool RequestFrameData() override;
-	virtual void UpdateRequestStatus(TArray<int32>& OutFramesCompleted) override;
-	virtual bool GetFrameData(int32 FrameIndex, FGeometryCacheMeshData& OutMeshData) override;
-	virtual int32 CancelRequests() override;
-	virtual const FGeometryCacheStreamStats& GetStreamStats() const override;
-	virtual void SetLimits(float MaxMemoryAllowed, float MaxCachedDuration) override;
+	virtual void GetMeshData(int32 FrameIndex, int32 ConcurrencyIndex, FGeometryCacheMeshData& OutMeshData) override;
 	//~ End IGeometryCacheStream Interface
 
-private:
-
-	void GetMeshData(int32 FrameIndex, int32 ConcurrencyIndex, FGeometryCacheMeshData& OutMeshData);
-	void LoadFrameData(int32 FrameIndex);
-	void UpdateFramesNeeded(int32 StartIndex, int32 NumFrames);
-	void IncrementMemoryStat(const FGeometryCacheMeshData& MeshData);
-	void DecrementMemoryStat(const FGeometryCacheMeshData& MeshData);
-
 	UGeometryCacheTrackAbcFile* AbcTrack;
-
-	TArray<int32> ReadIndices;
-	TArray<FGeometryCacheAbcStreamReadRequest*> ReadRequestsPool;
-
-	TArray<int32> FramesNeeded;
-	TArray<int32> FramesToBeCached;
-	TArray<FGeometryCacheAbcStreamReadRequest*> FramesRequested;
-
-	typedef TMap<int32, FGeometryCacheMeshData*> FFrameIndexToMeshData;
-	FFrameIndexToMeshData FramesAvailable;
-	FRWLock FramesAvailableLock;
-
-	std::atomic<bool> bCancellationRequested;
-
-	mutable FGeometryCacheStreamStats Stats;
 	FString Hash;
-	float SecondsPerFrame;
-	std::atomic<int32> LastAccessedFrameIndex;
-	int32 MaxCachedFrames;
-	float MaxCachedDuration;
-	float MaxMemAllowed;
-	float MemoryUsed;
-	std::atomic<bool> bCacheNeedsUpdate;
 };
\ No newline at end of file
diff --git a/Engine/Plugins/Experimental/GeometryCacheAbcFile/Source/GeometryCacheAbcFile/Private/GeometryCacheTrackAbcFile.cpp b/Engine/Plugins/Experimental/GeometryCacheAbcFile/Source/GeometryCacheAbcFile/Private/GeometryCacheTrackAbcFile.cpp
--- a/Engine/Plugins/Experimental/GeometryCacheAbcFile/Source/GeometryCacheAbcFile/Private/GeometryCacheTrackAbcFile.cpp
+++ b/Engine/Plugins/Experimental/GeometryCacheAbcFile/Source/GeometryCacheAbcFile/Private/GeometryCacheTrackAbcFile.cpp
@@ -116,6 +116,7 @@ void UGeometryCacheTrackAbcFile::Reset()
 {
 	AbcFile.Reset();
 	Hash.Empty();
+	AbcStream.Reset();
 
 	EndFrameIndex = 0;
 	Duration = 0.f;
@@ -219,14 +220,14 @@ bool UGeometryCacheTrackAbcFile::SetSourceFile(const FString& FilePath, UAbcImpo
 
 		// Register this Track and associated Stream with the GeometryCacheStreamer and prefetch the first frame
 		// The Stream ownership is passed to the Streamer
-		FGeometryCacheAbcStream* Stream = new FGeometryCacheAbcStream(this);
-		Streamer.RegisterTrack(this, Stream);
+		AbcStream.Reset(new FGeometryCacheAbcStream(this));
+		Streamer.RegisterTrack(this, AbcStream.Get());
 
 		const int32 InitialFrameIndex = FindSampleIndexFromTime(InitialTime, bIsLooping);
 		const float LookAhead = GetDefault<UGeometryCacheStreamerSettings>()->LookAheadBuffer;
 		const int32 NumFrames = FMath::CeilToInt(LookAhead / AbcFile->GetSecondsPerFrame());
 
-		Stream->Prefetch(InitialFrameIndex, NumFrames);
+		AbcStream->Prefetch(InitialFrameIndex, NumFrames);
 		GetMeshData(InitialFrameIndex, MeshData);
 
 		if (MeshData.Positions.Num() == 0)
@@ -316,4 +317,13 @@ FAbcFile& UGeometryCacheTrackAbcFile::GetAbcFile()
 	return *AbcFile.Get();
 }
 
+void UGeometryCacheTrackAbcFile::UpdateTime(float Time, bool bLooping)
+{
+	if (AbcStream)
+	{
+		int32 FrameIndex = FindSampleIndexFromTime(Time, bLooping);
+		AbcStream->UpdateCurrentFrameIndex(FrameIndex);
+	}
+}
+
 #undef LOCTEXT_NAMESPACE
diff --git a/Engine/Plugins/Experimental/GeometryCacheAbcFile/Source/GeometryCacheAbcFile/Public/GeometryCacheTrackAbcFile.h b/Engine/Plugins/Experimental/GeometryCacheAbcFile/Source/GeometryCacheAbcFile/Public/GeometryCacheTrackAbcFile.h
--- a/Engine/Plugins/Experimental/GeometryCacheAbcFile/Source/GeometryCacheAbcFile/Public/GeometryCacheTrackAbcFile.h
+++ b/Engine/Plugins/Experimental/GeometryCacheAbcFile/Source/GeometryCacheAbcFile/Public/GeometryCacheTrackAbcFile.h
@@ -4,6 +4,7 @@
 
 #include "CoreMinimal.h"
 #include "AbcFile.h"
+#include "GeometryCacheAbcStream.h"
 #include "GeometryCacheTrack.h"
 #include "GeometryCacheMeshData.h"
 
@@ -26,6 +27,7 @@ public:
 	virtual const bool UpdateBoundsData(const float Time, const bool bLooping, const bool bIsPlayingBackward, int32& InOutBoundsSampleIndex, FBox& OutBounds) override;
 	virtual const FGeometryCacheTrackSampleInfo& GetSampleInfo(float Time, const bool bLooping) override;
 	virtual bool GetMeshDataAtTime(float Time, FGeometryCacheMeshData& OutMeshData) override;
+	virtual void UpdateTime(float Time, bool bLooping) override;
 	//~ End UGeometryCacheTrack Interface.
 
 	bool SetSourceFile(const FString& FilePath, class UAbcImportSettings* AbcSettings, float InitialTime = 0.f, bool bIsLooping = true);
@@ -55,4 +57,6 @@ private:
 	FString Hash;
 
 	int32 EndFrameIndex;
+
+	TUniquePtr<FGeometryCacheAbcStream> AbcStream;
 };
diff --git a/Engine/Plugins/Importers/USDImporter/Source/GeometryCacheUSD/GeometryCacheUSD.build.cs b/Engine/Plugins/Importers/USDImporter/Source/GeometryCacheUSD/GeometryCacheUSD.build.cs
--- a/Engine/Plugins/Importers/USDImporter/Source/GeometryCacheUSD/GeometryCacheUSD.build.cs
+++ b/Engine/Plugins/Importers/USDImporter/Source/GeometryCacheUSD/GeometryCacheUSD.build.cs
@@ -12,6 +12,7 @@ public class GeometryCacheUSD : ModuleRules
 				"CoreUObject",
 				"Engine",
 				"GeometryCache",
+				"GeometryCacheStreamer",
 				"RenderCore",
 				"RHI"
 			}
@@ -20,8 +21,10 @@ public class GeometryCacheUSD : ModuleRules
 		PrivateDependencyModuleNames.AddRange(
 			new string[] {
 				"UnrealUSDWrapper",
-				"GeometryCacheStreamer"
+				"USDUtilities"
             }
         );
+
+		PrivateIncludePathModuleNames.Add("DerivedDataCache");
 	}
 }
diff --git a/Engine/Plugins/Importers/USDImporter/Source/GeometryCacheUSD/Private/GeometryCacheTrackUSD.cpp b/Engine/Plugins/Importers/USDImporter/Source/GeometryCacheUSD/Private/GeometryCacheTrackUSD.cpp
--- a/Engine/Plugins/Importers/USDImporter/Source/GeometryCacheUSD/Private/GeometryCacheTrackUSD.cpp
+++ b/Engine/Plugins/Importers/USDImporter/Source/GeometryCacheUSD/Private/GeometryCacheTrackUSD.cpp
@@ -19,6 +19,7 @@ void UGeometryCacheTrackUsd::BeginDestroy()
 	CurrentStage = UE::FUsdStage();
 
 	IGeometryCacheStreamer::Get().UnregisterTrack( this );
+	UsdStream.Reset();
 
 	Super::BeginDestroy();
 }
@@ -86,7 +87,7 @@ const FGeometryCacheTrackSampleInfo& UGeometryCacheTrackUsd::GetSampleInfo(float
 	if (CurrentSampleInfo.SampleTime == 0.0f && CurrentSampleInfo.NumVertices == 0 && CurrentSampleInfo.NumIndices == 0)
 	{
 		FGeometryCacheMeshData TempMeshData;
-		if ( GetMeshData( SampleIndex, TempMeshData ) )
+		if (GetMeshData(SampleIndex, TempMeshData))
 		{
 			CurrentSampleInfo = FGeometryCacheTrackSampleInfo(
 				Time,
@@ -137,7 +138,16 @@ void UGeometryCacheTrackUsd::Initialize(
 
 	Duration = ( float ) ( EndFrameIndex - StartFrameIndex );
 
-	FGeometryCacheUsdStream* Stream = new FGeometryCacheUsdStream( this, InReadFunc );
-	IGeometryCacheStreamer::Get().RegisterTrack( this, Stream );
-	Stream->Prefetch( StartFrameIndex );
+	UsdStream.Reset(new FGeometryCacheUsdStream(this, InReadFunc));
+	IGeometryCacheStreamer::Get().RegisterTrack(this, UsdStream.Get());
+	UsdStream->Prefetch(StartFrameIndex);
 }
+
+void UGeometryCacheTrackUsd::UpdateTime(float Time, bool bLooping)
+{
+	if (UsdStream)
+	{
+		int32 FrameIndex = FindSampleIndexFromTime(Time, bLooping);
+		UsdStream->UpdateCurrentFrameIndex(FrameIndex);
+	}
+}
\ No newline at end of file
diff --git a/Engine/Plugins/Importers/USDImporter/Source/GeometryCacheUSD/Private/GeometryCacheUSDStream.cpp b/Engine/Plugins/Importers/USDImporter/Source/GeometryCacheUSD/Private/GeometryCacheUSDStream.cpp
--- a/Engine/Plugins/Importers/USDImporter/Source/GeometryCacheUSD/Private/GeometryCacheUSDStream.cpp
+++ b/Engine/Plugins/Importers/USDImporter/Source/GeometryCacheUSD/Private/GeometryCacheUSDStream.cpp
@@ -1,284 +1,124 @@
 // Copyright Epic Games, Inc. All Rights Reserved.
 
 #include "GeometryCacheUSDStream.h"
-#include "Async/Async.h"
+#include "DerivedDataCacheInterface.h"
 #include "GeometryCacheMeshData.h"
-#include "Misc/ScopeLock.h"
+#include "GeometryCacheTrackUSD.h"
+#include "HAL/IConsoleManager.h"
+#include "Misc/CoreMisc.h"
+#include "USDGeomMeshConversion.h"
+
+static bool GUsdStreamCacheInDDC = true;
+static FAutoConsoleVariableRef CVarUsdStreamCacheInDDC(
+	TEXT("GeometryCache.Streamer.UsdStream.CacheInDDC"),
+	GUsdStreamCacheInDDC,
+	TEXT("Cache the streamed USD mesh data in the DDC"));
 
 static int32 kUsdReadConcurrency = 10;
 
-struct FGeometryCacheUsdStreamReadRequest
-{
-	FGeometryCacheMeshData* MeshData;
-	int32 ReadIndex;
-	int32 FrameIndex;
-	bool Completed;
-};
+// If UsdStream derived data needs to be rebuilt (new format, serialization
+// differences, etc.) replace the version GUID below with a new one.
+// In case of merge conflicts with DDC versions, you MUST generate a new GUID
+// and set this new GUID as the version.
+#define USDSTREAM_DERIVED_DATA_VERSION TEXT("040845A286FE431A846D95D9B77EE756")
 
-FGeometryCacheUsdStream::FGeometryCacheUsdStream(TWeakObjectPtr<UGeometryCacheTrackUsd> InUsdTrack, FReadUsdMeshFunction InReadFunc)
-: UsdTrack(InUsdTrack)
-, ReadFunc(InReadFunc)
-, bCancellationRequested(false)
+class FUsdStreamDDCUtils
 {
-	for (int32 Index = 0; Index < kUsdReadConcurrency; ++Index)
+private:
+	static const FString& GetUsdStreamDerivedDataVersion()
 	{
-		// Populate the ReadIndices. Note that it used as a stack
-		ReadIndices.Push(Index);
-
-		// Populate pool of reusable ReadRequests
-		ReadRequestsPool.Add(new FGeometryCacheUsdStreamReadRequest());
+		static FString CachedVersionString(USDSTREAM_DERIVED_DATA_VERSION);
+		return CachedVersionString;
 	}
-}
 
-FGeometryCacheUsdStream::~FGeometryCacheUsdStream()
-{
-	CancelRequests();
-
-	// Delete all the cached MeshData
-	FScopeLock Lock(&CriticalSection);
-	for (FFrameIndexToMeshData::TIterator Iterator = FramesAvailable.CreateIterator(); Iterator; ++Iterator)
+	static FString BuildDerivedDataKey(const FString& KeySuffix)
 	{
-		delete Iterator->Value;
+		return FDerivedDataCacheInterface::BuildCacheKey(TEXT("USDSTREAM_"), *GetUsdStreamDerivedDataVersion(), *KeySuffix);
 	}
 
-	// And ReadRequests from the pool
-	for (int32 Index = 0; Index < ReadRequestsPool.Num(); ++Index)
+public:
+	static FString GetUsdStreamDDCKey(const UE::FUsdStage& Stage, const FString& PrimPath, int32 FrameIndex)
 	{
-		delete ReadRequestsPool[Index];
-	}
-}
-
-int32 FGeometryCacheUsdStream::CancelRequests()
-{
-	TGuardValue<TAtomic<bool>, bool> CancellationRequested(bCancellationRequested, true);
+		FString PrimHash = UsdUtils::HashGeomMeshPrim(Stage, PrimPath, FrameIndex);
 
-	// Clear the FramesNeeded to prevent scheduling further reads
-	FramesNeeded.Empty();
-
-	// Wait for all read requests to complete
-	TArray<int32> CompletedFrames;
-	while (FramesRequested.Num())
-	{
-		UpdateRequestStatus(CompletedFrames);
-		if (FramesRequested.Num())
+		if (!PrimHash.IsEmpty())
 		{
-			FPlatformProcess::Sleep(0.01f);
+			return BuildDerivedDataKey(PrimHash);
 		}
+		return {};
 	}
+};
 
-	return CompletedFrames.Num();
+FGeometryCacheUsdStream::FGeometryCacheUsdStream(UGeometryCacheTrackUsd* InUsdTrack, FReadUsdMeshFunction InReadFunc)
+: FGeometryCacheStreamBase(
+	kUsdReadConcurrency,
+	FGeometryCacheStreamDetails{
+		InUsdTrack->GetEndFrameIndex() - InUsdTrack->GetStartFrameIndex(),
+		float((InUsdTrack->GetEndFrameIndex() - InUsdTrack->GetStartFrameIndex()) / InUsdTrack->CurrentStage.GetFramesPerSecond()),
+		float(1.0f / InUsdTrack->CurrentStage.GetFramesPerSecond()),
+		InUsdTrack->GetStartFrameIndex(),
+		InUsdTrack->GetEndFrameIndex()})
+, UsdTrack(InUsdTrack)
+, ReadFunc(InReadFunc)
+{
 }
 
-bool FGeometryCacheUsdStream::RequestFrameData()
+bool FGeometryCacheUsdStream::GetFrameData(int32 FrameIndex, FGeometryCacheMeshData& OutMeshData)
 {
-	check(IsInGameThread());
-
-	if (FramesNeeded.Num() == 0)
-	{
-		return false;
-	}
-
-	int32 FrameIndex = FramesNeeded[0];
-
-	// Don't schedule the same FrameIndex twice
-	for (int32 Index = 0; Index < FramesRequested.Num(); ++Index)
-	{
-		if (FramesRequested[Index]->FrameIndex == FrameIndex)
+	if (!FGeometryCacheStreamBase::GetFrameData(FrameIndex, OutMeshData))
+	{
+		// If the user requested a frame that isn't loaded yet, synchronously fetch it right away
+		// or else UGeometryCacheTrackUsd::GetSampleInfo may return invalid bounding boxes that may lead
+		// to issues, and wouldn't otherwise be updated until that frame is requested again.
+		// It may lead to a bit of stuttering when animating through an unloaded section with the sequencer
+		// or by dragging the Time property of the stage actor, but the alternative would be a spam of
+		// warnings on the output log and glitchy bounding boxes on the level
+		if (IsInGameThread())
 		{
-			return true;
+			LoadFrameData(FrameIndex);
+			return FGeometryCacheStreamBase::GetFrameData(FrameIndex, OutMeshData);
 		}
-	}
-
-	if (ReadIndices.Num() > 0)
-	{
-		// Get any ReadIndex available
-		const int32 ReadIndex = ReadIndices.Pop();
-
-		// Take the ReadRequest from the pool at ReadIndex and initialize it
-		FGeometryCacheUsdStreamReadRequest*& ReadRequest = ReadRequestsPool[ReadIndex];
-		ReadRequest->FrameIndex = FrameIndex;
-		ReadRequest->ReadIndex = ReadIndex;
-		ReadRequest->MeshData = new FGeometryCacheMeshData;
-		ReadRequest->Completed = false;
-
-		// Change the frame status from needed to requested
-		FramesNeeded.Remove(FrameIndex);
-		FramesRequested.Add(ReadRequest);
-
-		// Schedule asynchronous read of the MeshData
-		Async(
-#if WITH_EDITOR
-			EAsyncExecution::LargeThreadPool,
-#else
-			EAsyncExecution::ThreadPool,
-#endif // WITH_EDITOR
-			[this, ReadRequest]()
-			{
-				if (!bCancellationRequested)
-				{
-					ReadFunc( UsdTrack, ReadRequest->FrameIndex, *ReadRequest->MeshData );
-				}
-				ReadRequest->Completed = true;
-			});
-
-		return true;
-	}
-	return false;
-}
-
-void FGeometryCacheUsdStream::UpdateRequestStatus(TArray<int32>& OutFramesCompleted)
-{
-	check(IsInGameThread());
-
-	FScopeLock Lock(&CriticalSection);
-
-	UGeometryCacheTrackUsd* Track = UsdTrack.Get();
-
-	bool bCompletedARequest = false;
-
-	// Check the completion status of the read requests in progress
-	TArray<FGeometryCacheUsdStreamReadRequest*> CompletedRequests;
-	for (FGeometryCacheUsdStreamReadRequest* ReadRequest : FramesRequested)
-	{
-		if (ReadRequest->Completed)
+		else
 		{
-			// Queue for removal after iterating
-			CompletedRequests.Add(ReadRequest);
-
-			// Cache result of read for retrieval later
-			// #ueent_todo: Implement some sort of LRU cache to reduce memory usage
-			FramesAvailable.Add(ReadRequest->FrameIndex, ReadRequest->MeshData);
-
-			// Push back the ReadIndex for reuse
-			ReadIndices.Push(ReadRequest->ReadIndex);
-
-			// Output the completed frame
-			OutFramesCompleted.Add(ReadRequest->FrameIndex);
-
-			bCompletedARequest = true;
+			return false;
 		}
 	}
-
-	for (FGeometryCacheUsdStreamReadRequest* ReadRequest : CompletedRequests)
-	{
-		FramesRequested.Remove(ReadRequest);
-	}
-
-	// We're fully done fetching what we need from USD for now, we can drop the track's strong stage reference so that the stage
-	// can close if needed. The track can reopen the stage whenever needed though (it will also notify the user on the output
-	// log if it did so: check CreateGeometryCache in USDGeomMeshTranslator.cpp).
-	// This is not exactly ideal, as the closed stage may have changes that the reopened stage doesn't, but the
-	// alternative would leave the stage open indefinitely (as the transaction buffer will keep the track alive for undo/redo),
-	// which can lead to even more strange behavior (like changes that persist even if you close and open the stage, etc.).
-	// Whenever the track reopens the stage it will use the stage cache anyway, so if that stage is open for any other reason
-	// (e.g. we did an undo and the stage actor reopened it) then the operation will just retrieve the stage from the cache,
-	// which is cheap.
-	if ( bCompletedARequest && FramesNeeded.Num() == 0 && FramesRequested.Num() == 0 && Track )
-	{
-		Track->CurrentStage = UE::FUsdStage();
-	}
+	return true;
 }
 
-void FGeometryCacheUsdStream::Prefetch(int32 StartFrameIndex, int32 NumFrames)
+void FGeometryCacheUsdStream::GetMeshData(int32 FrameIndex, int32 ConcurrencyIndex, FGeometryCacheMeshData& OutMeshData)
 {
-	UGeometryCacheTrackUsd* ValidUsdTrack = UsdTrack.Get();
-	if ( !ValidUsdTrack )
-	{
-		return;
-	}
-
-	const int32 StartIndex = ValidUsdTrack->GetStartFrameIndex();
-	const int32 EndIndex = ValidUsdTrack->GetEndFrameIndex();
-	const int32 MaxNumFrames = EndIndex - StartIndex;
-
-	// Validate the number of frames to be loaded
-	if (NumFrames == 0)
-	{
-		// If no value specified, load the whole stream
-		NumFrames = MaxNumFrames;
-	}
-	else
-	{
-		NumFrames = FMath::Clamp(NumFrames, 1, MaxNumFrames);
-	}
+	TRACE_CPUPROFILER_EVENT_SCOPE(FGeometryCacheUsdStream::GetMeshData);
 
-	// Populate the list of frames needed to be loaded starting from given StartFrameIndex up to NumFrames or EndIndex
-	StartFrameIndex = FMath::Clamp(StartFrameIndex, StartIndex, EndIndex);
-	for (int32 Index = StartFrameIndex; NumFrames > 0 && Index < EndIndex; ++Index, --NumFrames)
+	// Get the mesh data straight from the Alembic file or from the DDC if it's already cached
+	if (GUsdStreamCacheInDDC)
 	{
-		FramesNeeded.Add(Index);
-	}
-
-	// End of the range might have been reached before the requested NumFrames so add the remaining frames starting from StartIndex
-	for (int32 Index = StartIndex; NumFrames > 0; ++Index, --NumFrames)
-	{
-		FramesNeeded.Add(Index);
-	}
-
-	if (FramesNeeded.Num() > 0)
-	{
-		// Force the first frame to be loaded and ready for retrieval
-		LoadFrameData(FramesNeeded[0]);
-		FramesNeeded.RemoveAt(0);
-	}
-}
-
-uint32 FGeometryCacheUsdStream::GetNumFramesNeeded()
-{
-	return FramesNeeded.Num();
-}
+		const FString& UsdPrimPath = UsdTrack->PrimPath;
+		const FString DerivedDataKey = FUsdStreamDDCUtils::GetUsdStreamDDCKey(UsdTrack->CurrentStage, UsdPrimPath, FrameIndex);
 
-bool FGeometryCacheUsdStream::GetFrameData(int32 FrameIndex, FGeometryCacheMeshData& OutMeshData)
-{
-	FScopeLock Lock(&CriticalSection);
-	if (FGeometryCacheMeshData** MeshDataPtr = FramesAvailable.Find(FrameIndex))
-	{
-		OutMeshData = **MeshDataPtr;
-		return true;
-	}
-	// If the user requested a frame that isn't loaded yet, synchronously fetch it right away
-	// or else UGeometryCacheTrackUsd::GetSampleInfo may return invalid bounding boxes that may lead
-	// to issues, and wouldn't otherwise be updated until that frame is requested again.
-	// It may lead to a bit of stuttering when animating through an unloaded section with the sequencer
-	// or by dragging the Time property of the stage actor, but the alternative would be a spam of
-	// warnings on the output log and glitchy bounding boxes on the level
-	else if ( IsInGameThread() )
-	{
-		LoadFrameData( FrameIndex );
-		return GetFrameData( FrameIndex, OutMeshData );
-	}
-
-	return false;
-}
+		if (!DerivedDataKey.IsEmpty())
+		{
+			TArray<uint8> DerivedData;
+			if (GetDerivedDataCacheRef().GetSynchronous(*DerivedDataKey, DerivedData, UsdPrimPath))
+			{
+				FMemoryReader Ar(DerivedData, /*bIsPersistent=*/ true);
+				Ar << OutMeshData;
+			}
+			else
+			{
+				ReadFunc(UsdTrack, FrameIndex, OutMeshData);
 
-void FGeometryCacheUsdStream::LoadFrameData(int32 FrameIndex)
-{
-	check(IsInGameThread());
+				FMemoryWriter Ar(DerivedData, true);
+				Ar << OutMeshData;
 
-	if ( FramesAvailable.Contains( FrameIndex ) )
-	{
-		return;
+				GetDerivedDataCacheRef().Put(*DerivedDataKey, DerivedData, UsdPrimPath);
+			}
+		}
+		// If the key is empty, it means the prim was invalid or not a mesh, so don't do anything
 	}
-
-	UGeometryCacheTrackUsd* Track = UsdTrack.Get();
-	if ( !Track )
+	else
 	{
-		return;
+		// Synchronously load the requested frame data
+		ReadFunc(UsdTrack, FrameIndex, OutMeshData);
 	}
-
-	// Synchronously load the requested frame data
-	FGeometryCacheMeshData* MeshData = new FGeometryCacheMeshData;
-	ReadFunc( UsdTrack, FrameIndex, *MeshData );
-
-	FramesAvailable.Add(FrameIndex, MeshData);
-}
-
-const FGeometryCacheStreamStats& FGeometryCacheUsdStream::GetStreamStats() const
-{
-	static FGeometryCacheStreamStats NoStats;
-	return NoStats;
-}
-
-void FGeometryCacheUsdStream::SetLimits(float MaxMemoryAllowed, float MaxCachedDuration)
-{
 }
diff --git a/Engine/Plugins/Importers/USDImporter/Source/GeometryCacheUSD/Private/GeometryCacheUSDStream.h b/Engine/Plugins/Importers/USDImporter/Source/GeometryCacheUSD/Private/GeometryCacheUSDStream.h
deleted file mode 100644
diff --git a/Engine/Plugins/Importers/USDImporter/Source/GeometryCacheUSD/Public/GeometryCacheTrackUSD.h b/Engine/Plugins/Importers/USDImporter/Source/GeometryCacheUSD/Public/GeometryCacheTrackUSD.h
--- a/Engine/Plugins/Importers/USDImporter/Source/GeometryCacheUSD/Public/GeometryCacheTrackUSD.h
+++ b/Engine/Plugins/Importers/USDImporter/Source/GeometryCacheUSD/Public/GeometryCacheTrackUSD.h
@@ -5,13 +5,12 @@
 #include "CoreMinimal.h"
 #include "GeometryCacheTrack.h"
 #include "GeometryCacheMeshData.h"
-
+#include "GeometryCacheTrackUSDTypes.h"
+#include "GeometryCacheUSDStream.h"
 #include "UsdWrappers/UsdStage.h"
 
 #include "GeometryCacheTrackUSD.generated.h"
 
-typedef TFunction< void( const TWeakObjectPtr<UGeometryCacheTrackUsd>, float Time, FGeometryCacheMeshData& ) > FReadUsdMeshFunction;
-
 /** GeometryCacheTrack for querying USD */
 UCLASS(collapsecategories, hidecategories = Object, BlueprintType, config = Engine)
 class GEOMETRYCACHEUSD_API UGeometryCacheTrackUsd : public UGeometryCacheTrack
@@ -39,6 +38,7 @@ public:
 	virtual const bool UpdateMeshData(const float Time, const bool bLooping, int32& InOutMeshSampleIndex, FGeometryCacheMeshData*& OutMeshData) override;
 	virtual const bool UpdateBoundsData(const float Time, const bool bLooping, const bool bIsPlayingBackward, int32& InOutBoundsSampleIndex, FBox& OutBounds) override;
 	virtual const FGeometryCacheTrackSampleInfo& GetSampleInfo(float Time, const bool bLooping) override;
+	virtual void UpdateTime(float Time, bool bLooping) override;
 	//~ End UGeometryCacheTrack Interface.
 
 	const int32 FindSampleIndexFromTime(const float Time, const bool bLooping) const;
@@ -62,4 +62,6 @@ public:
 private:
 	FGeometryCacheMeshData MeshData;
 	TArray<FGeometryCacheTrackSampleInfo> SampleInfos;
+
+	TUniquePtr<FGeometryCacheUsdStream> UsdStream;
 };
diff --git a/Engine/Plugins/Importers/USDImporter/Source/GeometryCacheUSD/Public/GeometryCacheTrackUSDTypes.h b/Engine/Plugins/Importers/USDImporter/Source/GeometryCacheUSD/Public/GeometryCacheTrackUSDTypes.h
new file mode 100644
--- /dev/null
+++ b/Engine/Plugins/Importers/USDImporter/Source/GeometryCacheUSD/Public/GeometryCacheTrackUSDTypes.h
@@ -0,0 +1,10 @@
+// Copyright Epic Games, Inc. All Rights Reserved.
+
+#pragma once
+
+#include "Templates/Function.h"
+
+class UGeometryCacheTrackUsd;
+struct FGeometryCacheMeshData;
+
+typedef TFunction< void( const TWeakObjectPtr<UGeometryCacheTrackUsd>, float Time, FGeometryCacheMeshData& ) > FReadUsdMeshFunction;
diff --git a/Engine/Plugins/Importers/USDImporter/Source/GeometryCacheUSD/Public/GeometryCacheUSDStream.h b/Engine/Plugins/Importers/USDImporter/Source/GeometryCacheUSD/Public/GeometryCacheUSDStream.h
new file mode 100644
--- /dev/null
+++ b/Engine/Plugins/Importers/USDImporter/Source/GeometryCacheUSD/Public/GeometryCacheUSDStream.h
@@ -0,0 +1,25 @@
+// Copyright Epic Games, Inc. All Rights Reserved.
+
+#pragma once
+
+#include "GeometryCacheStreamBase.h"
+#include "GeometryCacheTrackUSDTypes.h"
+
+class UGeometryCacheTrackUsd;
+
+class FGeometryCacheUsdStream : public FGeometryCacheStreamBase
+{
+public:
+	FGeometryCacheUsdStream(UGeometryCacheTrackUsd* InUsdTrack, FReadUsdMeshFunction InReadFunc);
+
+	//~ Begin IGeometryCacheStream Interface
+	virtual bool GetFrameData(int32 FrameIndex, FGeometryCacheMeshData& OutMeshData) override;
+
+protected:
+	virtual void GetMeshData(int32 FrameIndex, int32 ConcurrencyIndex, FGeometryCacheMeshData& OutMeshData) override;
+	//~ End IGeometryCacheStream Interface
+
+	UGeometryCacheTrackUsd* UsdTrack;
+	FReadUsdMeshFunction ReadFunc;
+};
+
diff --git a/Engine/Plugins/Importers/USDImporter/Source/USDUtilities/Private/USDGeomMeshConversion.cpp b/Engine/Plugins/Importers/USDImporter/Source/USDUtilities/Private/USDGeomMeshConversion.cpp
--- a/Engine/Plugins/Importers/USDImporter/Source/USDUtilities/Private/USDGeomMeshConversion.cpp
+++ b/Engine/Plugins/Importers/USDImporter/Source/USDUtilities/Private/USDGeomMeshConversion.cpp
@@ -15,6 +15,7 @@
 #include "USDTypesConversion.h"
 
 #include "UsdWrappers/SdfLayer.h"
+#include "UsdWrappers/SdfPath.h"
 #include "UsdWrappers/UsdPrim.h"
 #include "UsdWrappers/UsdStage.h"
 
@@ -2333,6 +2334,86 @@ void UsdUtils::ReplaceUnrealMaterialsWithBaked(
 	TraverseForMaterialReplacement( UsdStage, Root, Empty, {} );
 }
 
+FString UsdUtils::HashGeomMeshPrim( const UE::FUsdStage& Stage, const FString& PrimPath, double TimeCode )
+{
+	TRACE_CPUPROFILER_EVENT_SCOPE( UsdUtils::HashGeomMeshPrim );
+
+	using namespace pxr;
+
+	FScopedUsdAllocs Allocs;
+
+	UsdPrim UsdPrim = pxr::UsdPrim( Stage.GetPrimAtPath( UE::FSdfPath( *PrimPath ) ) );
+	if ( !UsdPrim )
+	{
+		return {};
+	}
+
+	UsdGeomMesh UsdMesh( UsdPrim );
+	if ( !UsdMesh )
+	{
+		return {};
+	}
+
+	FMD5 MD5;
+
+	if ( UsdAttribute Points = UsdMesh.GetPointsAttr() )
+	{
+		VtArray< GfVec3f > PointsArray;
+		Points.Get( &PointsArray, TimeCode );
+		MD5.Update( ( uint8* ) PointsArray.cdata(), PointsArray.size() * sizeof( GfVec3f ) );
+	}
+
+	if ( UsdAttribute NormalsAttribute = UsdMesh.GetNormalsAttr() )
+	{
+		VtArray< GfVec3f > Normals;
+		NormalsAttribute.Get( &Normals, TimeCode );
+		MD5.Update( ( uint8* ) Normals.cdata(), Normals.size() * sizeof( GfVec3f ) );
+	}
+
+	if ( UsdGeomPrimvar ColorPrimvar = UsdMesh.GetDisplayColorPrimvar() )
+	{
+		VtArray< GfVec3f > UsdColors;
+		ColorPrimvar.ComputeFlattened( &UsdColors, TimeCode );
+		MD5.Update( ( uint8* ) UsdColors.cdata(), UsdColors.size() * sizeof( GfVec3f ) );
+	}
+
+	if ( UsdGeomPrimvar OpacityPrimvar = UsdMesh.GetDisplayOpacityPrimvar() )
+	{
+		VtArray< float > UsdOpacities;
+		OpacityPrimvar.ComputeFlattened( &UsdOpacities );
+		MD5.Update( ( uint8* ) UsdOpacities.cdata(), UsdOpacities.size() * sizeof( float ) );
+	}
+
+	TArray< TUsdStore< UsdGeomPrimvar > > PrimvarsByUVIndex = UsdUtils::GetUVSetPrimvars( UsdMesh );
+	for ( int32 UVChannelIndex = 0; UVChannelIndex < PrimvarsByUVIndex.Num(); ++UVChannelIndex )
+	{
+		if ( !PrimvarsByUVIndex.IsValidIndex( UVChannelIndex ) )
+		{
+			break;
+		}
+
+		UsdGeomPrimvar& Primvar = PrimvarsByUVIndex[ UVChannelIndex ].Get();
+		if ( !Primvar )
+		{
+			continue;
+		}
+
+		VtArray< GfVec2f > UsdUVs;
+		Primvar.Get( &UsdUVs, TimeCode );
+		MD5.Update( ( uint8* ) UsdUVs.cdata(), UsdUVs.size() * sizeof( GfVec2f ) );
+	}
+
+	uint8 Digest[ 16 ];
+	MD5.Final( Digest );
+
+	FString Hash;
+	for ( int32 i = 0; i < 16; ++i )
+	{
+		Hash += FString::Printf( TEXT( "%02x" ), Digest[ i ] );
+	}
+	return Hash;
+}
+
 #undef LOCTEXT_NAMESPACE
 
 #endif // #if USE_USD_SDK
diff --git a/Engine/Plugins/Importers/USDImporter/Source/USDUtilities/Public/USDGeomMeshConversion.h b/Engine/Plugins/Importers/USDImporter/Source/USDUtilities/Public/USDGeomMeshConversion.h
--- a/Engine/Plugins/Importers/USDImporter/Source/USDUtilities/Public/USDGeomMeshConversion.h
+++ b/Engine/Plugins/Importers/USDImporter/Source/USDUtilities/Public/USDGeomMeshConversion.h
@@ -216,6 +216,12 @@ namespace UsdUtils
 	 *                                 Important because the `unrealMaterial` attributes will be used as a higher priority when determining material assignments
 	 */
 	USDUTILITIES_API void ReplaceUnrealMaterialsWithBaked( const UE::FUsdStage& Stage, const UE::FSdfLayer& LayerToAuthorIn, const TMap<FString, FString>& BakedMaterials, bool bIsAssetLayer, bool bUsePayload, bool bRemoveUnrealMaterials );
+
+	/**
+	 * Hashes the attributes of the GeomMesh at the given prim path on the Stage at TimeCode
+	 * If the prim path is invalid or not a GeomMesh, it will return an empty hash
+	 */
+	USDUTILITIES_API FString HashGeomMeshPrim( const UE::FUsdStage& Stage, const FString& PrimPath, double TimeCode );
 }
 
 #endif // #if USE_USD_SDK
