GeometryCache:
* Added support for frame interpolation to streaming Alembic GeometryCache, needed for subframe sampling and motion blur.

diff --git a/Engine/Plugins/Experimental/AlembicImporter/Source/AlembicLibrary/Private/AbcFile.cpp b/Engine/Plugins/Experimental/AlembicImporter/Source/AlembicLibrary/Private/AbcFile.cpp
--- a/Engine/Plugins/Experimental/AlembicImporter/Source/AlembicLibrary/Private/AbcFile.cpp
+++ b/Engine/Plugins/Experimental/AlembicImporter/Source/AlembicLibrary/Private/AbcFile.cpp
@@ -641,7 +641,7 @@ int32 FAbcFile::GetFrameIndex(float Time)
 {
 	if (SecondsPerFrame > 0.f)
 	{
-		int32 FrameIndex = StartFrameIndex + FMath::RoundToInt(Time / SecondsPerFrame);
+		int32 FrameIndex = StartFrameIndex + FMath::FloorToInt(Time / SecondsPerFrame);
 		return FMath::Clamp(FrameIndex, StartFrameIndex, EndFrameIndex);
 	}
 	return 0;
diff --git a/Engine/Plugins/Experimental/GeometryCache/Source/GeometryCache/Classes/GeometryCacheTrack.h b/Engine/Plugins/Experimental/GeometryCache/Source/GeometryCache/Classes/GeometryCacheTrack.h
--- a/Engine/Plugins/Experimental/GeometryCache/Source/GeometryCache/Classes/GeometryCacheTrack.h
+++ b/Engine/Plugins/Experimental/GeometryCache/Source/GeometryCache/Classes/GeometryCacheTrack.h
@@ -168,7 +168,7 @@ struct GEOMETRYCACHE_API FGeometryCacheTrackSampleInfo
 	int32 NumVertices;
 	int32 NumIndices;
 
-	FGeometryCacheTrackSampleInfo() : SampleTime(0.0f), NumVertices(0), NumIndices(0) {}
+	FGeometryCacheTrackSampleInfo() : SampleTime(0.0f), BoundingBox(EForceInit::ForceInit), NumVertices(0), NumIndices(0) {}
 
 	FGeometryCacheTrackSampleInfo(float SetSampleTime, FBox SetBoundingBox, int32 SetNumVertices, int32 SetNumIndices) :
 		SampleTime(SetSampleTime), BoundingBox(SetBoundingBox), NumVertices(SetNumVertices), NumIndices(SetNumIndices) {}
diff --git a/Engine/Plugins/Experimental/GeometryCacheAbcFile/Source/GeometryCacheAbcFile/Private/GeometryCacheAbcFileSceneProxy.cpp b/Engine/Plugins/Experimental/GeometryCacheAbcFile/Source/GeometryCacheAbcFile/Private/GeometryCacheAbcFileSceneProxy.cpp
--- a/Engine/Plugins/Experimental/GeometryCacheAbcFile/Source/GeometryCacheAbcFile/Private/GeometryCacheAbcFileSceneProxy.cpp
+++ b/Engine/Plugins/Experimental/GeometryCacheAbcFile/Source/GeometryCacheAbcFile/Private/GeometryCacheAbcFileSceneProxy.cpp
@@ -2,6 +2,7 @@
 
 #include "GeometryCacheAbcFileSceneProxy.h"
 #include "GeometryCacheAbcFileComponent.h"
+#include "GeometryCacheHelpers.h"
 #include "GeometryCacheTrackAbcFile.h"
 
 FGeometryCacheAbcFileSceneProxy::FGeometryCacheAbcFileSceneProxy(UGeometryCacheAbcFileComponent* Component)
@@ -35,7 +36,10 @@ bool FGeomCacheTrackAbcFileProxy::GetMeshData(int32 SampleIndex, FGeometryCacheM
 
 bool FGeomCacheTrackAbcFileProxy::IsTopologyCompatible(int32 SampleIndexA, int32 SampleIndexB)
 {
-	// No support for interpolation for now (assume the topology is variable)
+	if (UGeometryCacheTrackAbcFile* AbcTrack = Cast<UGeometryCacheTrackAbcFile>(Track))
+	{
+		return AbcTrack->IsTopologyCompatible(SampleIndexA, SampleIndexB);
+	}
 	return false;
 }
 
@@ -53,7 +57,16 @@ void FGeomCacheTrackAbcFileProxy::FindSampleIndexesFromTime(float Time, bool bLo
 		int32 LastFrameIndex = AbcTrack->GetEndFrameIndex();
 		OutFrameIndex = ThisFrameIndex;
 		OutNextFrameIndex = OutFrameIndex + 1;
-		InInterpolationFactor = 0.f;
+
+		float AdjustedTime = Time;
+		if (bLooping)
+		{
+			AdjustedTime = GeometyCacheHelpers::WrapAnimationTime(Time, AbcTrack->GetDuration());
+		}
+
+		// Time at ThisFrameIndex with index normalized to 0
+		const float FrameIndexTime = (ThisFrameIndex - AbcTrack->GetAbcFile().GetStartFrameIndex()) * AbcTrack->GetAbcFile().GetSecondsPerFrame();
+		InInterpolationFactor = (AdjustedTime - FrameIndexTime) / AbcTrack->GetAbcFile().GetSecondsPerFrame();
 
 		// If playing backwards the logical order of previous and next is reversed
 		if (bIsPlayingBackwards)
diff --git a/Engine/Plugins/Experimental/GeometryCacheAbcFile/Source/GeometryCacheAbcFile/Private/GeometryCacheTrackAbcFile.cpp b/Engine/Plugins/Experimental/GeometryCacheAbcFile/Source/GeometryCacheAbcFile/Private/GeometryCacheTrackAbcFile.cpp
--- a/Engine/Plugins/Experimental/GeometryCacheAbcFile/Source/GeometryCacheAbcFile/Private/GeometryCacheTrackAbcFile.cpp
+++ b/Engine/Plugins/Experimental/GeometryCacheAbcFile/Source/GeometryCacheAbcFile/Private/GeometryCacheTrackAbcFile.cpp
@@ -62,7 +62,8 @@ namespace UE::GeometryCacheTrackAbcFile::Private
 }
 
 UGeometryCacheTrackAbcFile::UGeometryCacheTrackAbcFile()
-: EndFrameIndex(0)
+: StartFrameIndex(0)
+, EndFrameIndex(0)
 {
 }
 
@@ -171,6 +172,12 @@ bool UGeometryCacheTrackAbcFile::SetSourceFile(const FString& FilePath, UAbcImpo
 			return false;
 		}
 
+		// Automatically set the FrameEnd the same way as it's computed in FAbcImporter::GetEndFrameIndex when importing as GeometryCache to have same duration
+		if (AbcSettings->SamplingSettings.FrameEnd == 0)
+		{
+			AbcSettings->SamplingSettings.FrameEnd = FMath::Max(AbcFile->GetMaxFrameIndex() - 1, 1);
+		}
+
 		Result = AbcFile->Import(AbcSettings);
 
 		// The hash is composed of the Alembic file hash and the settings hash used to import it
@@ -187,7 +194,8 @@ bool UGeometryCacheTrackAbcFile::SetSourceFile(const FString& FilePath, UAbcImpo
 		FString SettingsHash = UE::GeometryCacheTrackAbcFile::Private::ComputeSettingsHash(AbcSettings);
 		Hash = AbcHash + TEXT("_") + SettingsHash;
 
-		// Set the end frame after import since it might have been modified due to validation at import
+		// Set the start/end frame after import since it might have been modified due to validation at import
+		StartFrameIndex = AbcSettings->SamplingSettings.FrameStart;
 		EndFrameIndex = AbcSettings->SamplingSettings.FrameEnd;
 
 		if (Result != EAbcImportError::AbcImportError_NoError)
@@ -257,24 +265,59 @@ const int32 UGeometryCacheTrackAbcFile::FindSampleIndexFromTime(const float Time
 
 const FGeometryCacheTrackSampleInfo& UGeometryCacheTrackAbcFile::GetSampleInfo(float Time, bool bLooping)
 {
-	float SampleTime = Time;
-	if (bLooping)
+	if (SampleInfos.Num() == 0)
+	{
+		if (EndFrameIndex > StartFrameIndex)
+		{
+			// +1 because the SampleInfo for EndFrameIndex is also needed
+			SampleInfos.SetNum(EndFrameIndex - StartFrameIndex + 1);
+		}
+		else
+		{
+			return FGeometryCacheTrackSampleInfo::EmptySampleInfo;
+		}
+	}
+
+	// The sample info index must start from 0, while the sample index is between the range of the animation
+	const int32 SampleIndex = FindSampleIndexFromTime(Time, bLooping);
+	const int32 SampleInfoIndex = SampleIndex - StartFrameIndex;
+
+	FGeometryCacheTrackSampleInfo& CurrentSampleInfo = SampleInfos[SampleInfoIndex];
+
+	if (CurrentSampleInfo.SampleTime == 0.0f && CurrentSampleInfo.NumVertices == 0 && CurrentSampleInfo.NumIndices == 0)
 	{
-		SampleTime = GeometyCacheHelpers::WrapAnimationTime(Time, Duration);
+		if (GetMeshData(SampleIndex, MeshData))
+		{
+			CurrentSampleInfo = FGeometryCacheTrackSampleInfo(
+				Time,
+				(FBox) MeshData.BoundingBox,
+				MeshData.Positions.Num(),
+				MeshData.Indices.Num()
+			);
+		}
 	}
 
-	// Update the mesh data as required
-	int32 ThisSampleIndex = FindSampleIndexFromTime(SampleTime, bLooping);
-	GetMeshData(ThisSampleIndex, MeshData);
+	return CurrentSampleInfo;
+}
 
-	SampleInfo = FGeometryCacheTrackSampleInfo(
-		SampleTime,
-		(FBox)MeshData.BoundingBox,
-		MeshData.Positions.Num(),
-		MeshData.Indices.Num()
-	);
+const FGeometryCacheTrackSampleInfo& UGeometryCacheTrackAbcFile::GetSampleInfo(int32 FrameIndex)
+{
+	if (AbcFile)
+	{
+		// FrameIndex is normalized to 0
+		return GetSampleInfo((FrameIndex - StartFrameIndex) * AbcFile->GetSecondsPerFrame(), false);
+	}
+	return FGeometryCacheTrackSampleInfo::EmptySampleInfo;
+}
 
-	return SampleInfo;
+bool UGeometryCacheTrackAbcFile::IsTopologyCompatible(int32 FrameA, int32 FrameB)
+{
+	// FrameA/B could be -1, meaning invalid frame
+	if (SampleInfos.Num() > 0 && FrameA >= 0 && FrameB >= 0)
+	{
+		return GetSampleInfo(FrameA).NumVertices == GetSampleInfo(FrameB).NumVertices;
+	}
+	return false;
 }
 
 bool UGeometryCacheTrackAbcFile::GetMeshDataAtTime(float Time, FGeometryCacheMeshData& OutMeshData)
diff --git a/Engine/Plugins/Experimental/GeometryCacheAbcFile/Source/GeometryCacheAbcFile/Public/GeometryCacheTrackAbcFile.h b/Engine/Plugins/Experimental/GeometryCacheAbcFile/Source/GeometryCacheAbcFile/Public/GeometryCacheTrackAbcFile.h
--- a/Engine/Plugins/Experimental/GeometryCacheAbcFile/Source/GeometryCacheAbcFile/Public/GeometryCacheTrackAbcFile.h
+++ b/Engine/Plugins/Experimental/GeometryCacheAbcFile/Source/GeometryCacheAbcFile/Public/GeometryCacheTrackAbcFile.h
@@ -45,17 +45,22 @@ public:
 
 	const FString& GetAbcTrackHash() const { return Hash; }
 
+	bool IsTopologyCompatible(int32 FrameA, int32 FrameB);
+
 private:
+	const FGeometryCacheTrackSampleInfo& GetSampleInfo(int32 FrameIndex);
+
 	void Reset();
 	void ShowNotification(const FText& Text);
 
 private:
 	FGeometryCacheMeshData MeshData;
-	FGeometryCacheTrackSampleInfo SampleInfo;
+	TArray<FGeometryCacheTrackSampleInfo> SampleInfos;
 	TUniquePtr<FAbcFile> AbcFile;
 	FString SourceFile;
 	FString Hash;
 
+	int32 StartFrameIndex;
 	int32 EndFrameIndex;
 
 	TUniquePtr<FGeometryCacheAbcStream> AbcStream;
