Streaming GeometryCache:
* Fixed Abc stream didn't reset its sample infos when reloading Alembic file, which prevented updating its bounding boxes.

diff --git a/Engine/Plugins/Experimental/GeometryCacheAbcFile/Source/GeometryCacheAbcFile/Private/GeometryCacheTrackAbcFile.cpp b/Engine/Plugins/Experimental/GeometryCacheAbcFile/Source/GeometryCacheAbcFile/Private/GeometryCacheTrackAbcFile.cpp
--- a/Engine/Plugins/Experimental/GeometryCacheAbcFile/Source/GeometryCacheAbcFile/Private/GeometryCacheTrackAbcFile.cpp
+++ b/Engine/Plugins/Experimental/GeometryCacheAbcFile/Source/GeometryCacheAbcFile/Private/GeometryCacheTrackAbcFile.cpp
@@ -118,6 +118,7 @@ void UGeometryCacheTrackAbcFile::Reset()
 	AbcFile.Reset();
 	Hash.Empty();
 	AbcStream.Reset();
+	SampleInfos.Reset();
 
 	EndFrameIndex = 0;
 	Duration = 0.f;
