GeometryCacheAbcFile: Migrate to the UE5 DDC API

diff --git a/Engine/Plugins/Experimental/GeometryCacheAbcFile/Source/GeometryCacheAbcFile/GeometryCacheAbcFile.Build.cs b/Engine/Plugins/Experimental/GeometryCacheAbcFile/Source/GeometryCacheAbcFile/GeometryCacheAbcFile.Build.cs
--- a/Engine/Plugins/Experimental/GeometryCacheAbcFile/Source/GeometryCacheAbcFile/GeometryCacheAbcFile.Build.cs
+++ b/Engine/Plugins/Experimental/GeometryCacheAbcFile/Source/GeometryCacheAbcFile/GeometryCacheAbcFile.Build.cs
@@ -29,6 +29,9 @@ public class GeometryCacheAbcFile : ModuleRules
             }
         );
 
-		PrivateIncludePathModuleNames.Add("DerivedDataCache");
+		if (Target.bBuildWithEditorOnlyData)
+		{
+			PrivateDependencyModuleNames.Add("DerivedDataCache");
+		}
 	}
 }
diff --git a/Engine/Plugins/Experimental/GeometryCacheAbcFile/Source/GeometryCacheAbcFile/Private/GeometryCacheAbcStream.cpp b/Engine/Plugins/Experimental/GeometryCacheAbcFile/Source/GeometryCacheAbcFile/Private/GeometryCacheAbcStream.cpp
--- a/Engine/Plugins/Experimental/GeometryCacheAbcFile/Source/GeometryCacheAbcFile/Private/GeometryCacheAbcStream.cpp
+++ b/Engine/Plugins/Experimental/GeometryCacheAbcFile/Source/GeometryCacheAbcFile/Private/GeometryCacheAbcStream.cpp
@@ -3,7 +3,8 @@
 #include "GeometryCacheAbcStream.h"
 #include "AbcFile.h"
 #include "AbcUtilities.h"
-#include "DerivedDataCacheInterface.h"
+#include "DerivedDataCache.h"
+#include "DerivedDataRequestOwner.h"
 #include "GeometryCacheTrackAbcFile.h"
 #include "HAL/IConsoleManager.h"
 #include "Serialization/MemoryReader.h"
@@ -37,20 +38,19 @@ struct FGeometryCacheAbcStreamReadRequest
 // differences, etc.) replace the version GUID below with a new one.
 // In case of merge conflicts with DDC versions, you MUST generate a new GUID
 // and set this new GUID as the version.
-#define ABCSTREAM_DERIVED_DATA_VERSION TEXT("75887BC18F674774A53F91E4A57709D6")
+#define ABCSTREAM_DERIVED_DATA_VERSION TEXTVIEW("75887BC18F674774A53F91E4A57709D6")
 
 class FAbcStreamDDCUtils
 {
 private:
-	static const FString& GetAbcStreamDerivedDataVersion()
+	static FStringView GetAbcStreamDerivedDataVersion()
 	{
-		static FString CachedVersionString(ABCSTREAM_DERIVED_DATA_VERSION);
-		return CachedVersionString;
+		return ABCSTREAM_DERIVED_DATA_VERSION;
 	}
 
-	static FString BuildDerivedDataKey(const FString& KeySuffix)
+	static UE::DerivedData::FCacheKey BuildDerivedDataKey(FStringView KeySuffix)
 	{
-		return FDerivedDataCacheInterface::BuildCacheKey(TEXT("ABCSTREAM_"), *GetAbcStreamDerivedDataVersion(), *KeySuffix);
+		return UE::DerivedData::BuildLegacyCacheKey(TEXTVIEW("ABCSTREAM_"), GetAbcStreamDerivedDataVersion(), KeySuffix);
 	}
 
 	static FString BuildAbcStreamDerivedDataKeySuffix(UGeometryCacheTrackAbcFile* AbcTrack, int32 FrameIndex)
@@ -59,7 +59,7 @@ private:
 	}
 
 public:
-	static FString GetAbcStreamDDCKey(UGeometryCacheTrackAbcFile* AbcTrack, int32 FrameIndex)
+	static UE::DerivedData::FCacheKey GetAbcStreamDDCKey(UGeometryCacheTrackAbcFile* AbcTrack, int32 FrameIndex)
 	{
 		return BuildDerivedDataKey(BuildAbcStreamDerivedDataKeySuffix(AbcTrack, FrameIndex));
 	}
@@ -84,24 +84,39 @@ void FGeometryCacheAbcStream::GetMeshData(int32 FrameIndex, int32 ConcurrencyInd
 	// Get the mesh data straight from the Alembic file or from the DDC if it's already cached
 	if (GAbcStreamCacheInDDC)
 	{
-		const FString DerivedDataKey = FAbcStreamDDCUtils::GetAbcStreamDDCKey(AbcTrack, FrameIndex);
-		const FString& AbcFile = AbcTrack->GetSourceFile();
+		using namespace UE;
+		using namespace UE::DerivedData;
 
-		TArray<uint8> DerivedData;
-		if (GetDerivedDataCacheRef().GetSynchronous(*DerivedDataKey, DerivedData, AbcFile))
-		{
-			FMemoryReader Ar(DerivedData, /*bIsPersistent=*/ true);
-			Ar << MeshData;
-		}
-		else
+		const FCacheKey DerivedDataKey = FAbcStreamDDCUtils::GetAbcStreamDDCKey(AbcTrack, FrameIndex);
+		const FSharedString AbcFile(AbcTrack->GetSourceFile());
+
+		FRequestOwner BlockingOwner(EPriority::Blocking);
+		GetCache().GetValue({{AbcFile, DerivedDataKey}}, BlockingOwner, [this, FrameIndex, ConcurrencyIndex, &MeshData](FCacheGetValueResponse&& Response)
 		{
-			FAbcUtilities::GetFrameMeshData(AbcTrack->GetAbcFile(), FrameIndex, MeshData, ConcurrencyIndex);
+			if (Response.Status == EStatus::Ok)
+			{
+				FSharedBuffer DerivedData = Response.Value.GetData().Decompress();
+				FMemoryReaderView Ar(DerivedData, /*bIsPersistent=*/ true);
+				Ar << MeshData;
+			}
+			else
+			{
+				FAbcUtilities::GetFrameMeshData(AbcTrack->GetAbcFile(), FrameIndex, MeshData, ConcurrencyIndex);
 
-			FMemoryWriter Ar(DerivedData, true);
-			Ar << MeshData;
+				FValue Value;
+				{
+					TArray<uint8> DerivedData;
+					FMemoryWriter Ar(DerivedData, true);
+					Ar << MeshData;
+					Value = FValue::Compress(MakeSharedBufferFromArray(MoveTemp(DerivedData)));
+				}
 
-			GetDerivedDataCacheRef().Put(*DerivedDataKey, DerivedData, AbcFile);
-		}
+				FRequestOwner AsyncOwner(EPriority::Normal);
+				GetCache().PutValue({{Response.Name, Response.Key, Value}}, AsyncOwner);
+				AsyncOwner.KeepAlive();
+			}
+		});
+		BlockingOwner.Wait();
 	}
 	else
 	{
